NNE - Fixing UE Build Health Issue 242291 (Caused by renaming the plugin)

diff --git a/Engine/Plugins/Experimental/NNE/Source/NNXUtils/NNXUtils.Build.cs b/Engine/Plugins/Experimental/NNE/Source/NNXUtils/NNXUtils.Build.cs
--- a/Engine/Plugins/Experimental/NNE/Source/NNXUtils/NNXUtils.Build.cs
+++ b/Engine/Plugins/Experimental/NNE/Source/NNXUtils/NNXUtils.Build.cs
@@ -13,7 +13,7 @@ public class NNXUtils : ModuleRules
 
 		PublicIncludePaths.AddRange(
 			new string[] {
-				Path.GetFullPath(Path.Combine(EngineDirectory, "Plugins/Experimental/NNX/Source/NNXCore/Private")),
+				Path.GetFullPath(Path.Combine(EngineDirectory, "Plugins/Experimental/NNE/Source/NNXCore/Private")),
 			}
 		);
 
