NNE load model weights to FTensor for access

Problem:
* Weights need to be accessible by the model by scheduling code and for upload to GPU after the model is loaded

Solution:
* Store weights as tensor with CPU data when loading model.

Tests:
* All NNX tests pass under win64 editor (appart from one also red on main, being discussed)

diff --git a/Engine/Plugins/Experimental/NNE/Source/NNXCore/Public/NNXTypes.h b/Engine/Plugins/Experimental/NNE/Source/NNXCore/Public/NNXTypes.h
--- a/Engine/Plugins/Experimental/NNE/Source/NNXCore/Public/NNXTypes.h
+++ b/Engine/Plugins/Experimental/NNE/Source/NNXCore/Public/NNXTypes.h
@@ -231,8 +231,9 @@ class FTensor : public FTensorDescBase
 {
 protected:
 	FTensorShape Shape;
-	uint32 Volume = 0;
+	TArray<uint8> PreparedData;
 	uint64 DataSize = 0;
+	uint32 Volume = 0;
 
 public:
 	const FTensorShape& GetShape() const
@@ -240,6 +241,23 @@ public:
 		return Shape;
 	}
 
+	TConstArrayView<uint8> GetPreparedData() const
+	{
+		return PreparedData;
+	}
+
+	void SetPreparedData(TConstArrayView<uint8> Data)
+	{
+		check(Data.Num() == DataSize);
+		PreparedData.Empty();
+		PreparedData.Append(Data);
+	}
+
+	bool HasPreparedData()
+	{
+		return !PreparedData.IsEmpty();
+	}
+
 	uint32 GetVolume() const
 	{
 		return Volume;
diff --git a/Engine/Plugins/Experimental/NNE/Source/NNXRuntimeRDG/Private/Hlsl/NNIRuntimeRDGConv.cpp b/Engine/Plugins/Experimental/NNE/Source/NNXRuntimeRDG/Private/Hlsl/NNIRuntimeRDGConv.cpp
--- a/Engine/Plugins/Experimental/NNE/Source/NNXRuntimeRDG/Private/Hlsl/NNIRuntimeRDGConv.cpp
+++ b/Engine/Plugins/Experimental/NNE/Source/NNXRuntimeRDG/Private/Hlsl/NNIRuntimeRDGConv.cpp
@@ -105,7 +105,7 @@ namespace UE::NNIRuntimeRDG::Private::Hlsl
 			return true;
 		}
 
-		virtual void Dispatch(FRDGBuilder& GraphBuilder, TConstArrayView<NNX::FTensorRDG> InputTensors, TConstArrayView<NNX::FTensorRDG> OutputTensors) override
+		virtual void Dispatch(FRDGBuilder& GraphBuilder, TConstArrayView<NNX::FTensorRDGRef> InputTensors, TConstArrayView<NNX::FTensorRDGRef> OutputTensors) override
 		{
 			using namespace UE::NNEHlslShaders::Internal;
 
@@ -114,16 +114,20 @@ namespace UE::NNIRuntimeRDG::Private::Hlsl
 
 			check(InputTensors.Num() >= 2 && InputTensors.Num() <= 3);
 			check(OutputTensors.Num() == 1);
-
-			const NNX::FTensorRDG& Input = InputTensors[0];
-			const NNX::FTensorRDG& Weights = InputTensors[1];
-			const NNX::FTensorRDG& Output = OutputTensors[0];
+			check(InputTensors[0] != nullptr);
+			check(InputTensors[1] != nullptr);
+			check(OutputTensors[0] != nullptr);
+			
+			const NNX::FTensorRDG& Input = *InputTensors[0];
+			const NNX::FTensorRDG& Weights = *InputTensors[1];
+			const NNX::FTensorRDG& Output = *OutputTensors[0];
 			const NNX::FTensorRDG* Bias = nullptr;
 			bool HasBias = false;
 
 			if (InputTensors.Num() == 3) {
 				HasBias = true;
-				Bias = &(InputTensors[2]);
+				check(InputTensors[2] != nullptr);
+				Bias = InputTensors[2];
 			}
 
 			check(Input.GetShape().Rank() > 2);
diff --git a/Engine/Plugins/Experimental/NNE/Source/NNXRuntimeRDG/Private/Hlsl/NNIRuntimeRDGConvTranspose.cpp b/Engine/Plugins/Experimental/NNE/Source/NNXRuntimeRDG/Private/Hlsl/NNIRuntimeRDGConvTranspose.cpp
--- a/Engine/Plugins/Experimental/NNE/Source/NNXRuntimeRDG/Private/Hlsl/NNIRuntimeRDGConvTranspose.cpp
+++ b/Engine/Plugins/Experimental/NNE/Source/NNXRuntimeRDG/Private/Hlsl/NNIRuntimeRDGConvTranspose.cpp
@@ -110,7 +110,7 @@ namespace UE::NNIRuntimeRDG::Private::Hlsl
 			return true;
 		}
 
-		virtual void Dispatch(FRDGBuilder& GraphBuilder, TConstArrayView<NNX::FTensorRDG> InputTensors, TConstArrayView<NNX::FTensorRDG> OutputTensors) override
+		virtual void Dispatch(FRDGBuilder& GraphBuilder, TConstArrayView<NNX::FTensorRDGRef> InputTensors, TConstArrayView<NNX::FTensorRDGRef> OutputTensors) override
 		{
 			using namespace UE::NNEHlslShaders::Internal;
 
@@ -119,16 +119,20 @@ namespace UE::NNIRuntimeRDG::Private::Hlsl
 
 			check(InputTensors.Num() >= 2 && InputTensors.Num() <= 3);
 			check(OutputTensors.Num() == 1);
+			check(InputTensors[0] != nullptr);
+			check(InputTensors[1] != nullptr);
+			check(OutputTensors[0] != nullptr);
 
-			const NNX::FTensorRDG& Input = InputTensors[0];
-			const NNX::FTensorRDG& Weights = InputTensors[1];
-			const NNX::FTensorRDG& Output = OutputTensors[0];
+			const NNX::FTensorRDG& Input = *InputTensors[0];
+			const NNX::FTensorRDG& Weights = *InputTensors[1];
+			const NNX::FTensorRDG& Output = *OutputTensors[0];
 			const NNX::FTensorRDG* Bias = nullptr;
 			bool HasBias = false;
 
 			if (InputTensors.Num() == 3) {
 				HasBias = true;
-				Bias = &(InputTensors[2]);
+				check(InputTensors[2] != nullptr);
+				Bias = InputTensors[2];
 			}
 
 			check(Input.GetShape().Rank() > 2);
diff --git a/Engine/Plugins/Experimental/NNE/Source/NNXRuntimeRDG/Private/Hlsl/NNIRuntimeRDGElementWiseBinary.cpp b/Engine/Plugins/Experimental/NNE/Source/NNXRuntimeRDG/Private/Hlsl/NNIRuntimeRDGElementWiseBinary.cpp
--- a/Engine/Plugins/Experimental/NNE/Source/NNXRuntimeRDG/Private/Hlsl/NNIRuntimeRDGElementWiseBinary.cpp
+++ b/Engine/Plugins/Experimental/NNE/Source/NNXRuntimeRDG/Private/Hlsl/NNIRuntimeRDGElementWiseBinary.cpp
@@ -64,15 +64,17 @@ namespace UE::NNIRuntimeRDG::Private::Hlsl
 			return true;
 		}
 
-		virtual void Dispatch(FRDGBuilder& GraphBuilder, TConstArrayView<NNX::FTensorRDG> InputTensors, TConstArrayView<NNX::FTensorRDG> OutputTensors) override
+		virtual void Dispatch(FRDGBuilder& GraphBuilder, TConstArrayView<NNX::FTensorRDGRef> InputTensors, TConstArrayView<NNX::FTensorRDGRef> OutputTensors) override
 		{
 			check(InputTensors.Num() == 2);
 			check(OutputTensors.Num() == 1);
+			check(InputTensors[0] != nullptr);
+			check(InputTensors[1] != nullptr);
+			check(OutputTensors[0] != nullptr);
+			const NNX::FTensorRDG& LHSInput = *InputTensors[0];
+			const NNX::FTensorRDG& RHSInput = *InputTensors[1];
+			const NNX::FTensorRDG& Output = *OutputTensors[0];
 
-			const NNX::FTensorRDG& LHSInput = InputTensors[0];
-			const NNX::FTensorRDG& RHSInput = InputTensors[1];
-			const NNX::FTensorRDG& Output = OutputTensors[0];
-			
 			FRDGBufferSRVRef LHSInputSRV = GraphBuilder.CreateSRV(FRDGBufferSRVDesc(LHSInput.GetBuffer(), PF_R32_FLOAT));
 			FRDGBufferSRVRef RHSInputSRV = GraphBuilder.CreateSRV(FRDGBufferSRVDesc(RHSInput.GetBuffer(), PF_R32_FLOAT));
 			FRDGBufferUAVRef OutputUAV = GraphBuilder.CreateUAV(FRDGBufferUAVDesc(Output.GetBuffer(), PF_R32_FLOAT));
diff --git a/Engine/Plugins/Experimental/NNE/Source/NNXRuntimeRDG/Private/Hlsl/NNIRuntimeRDGElementWiseUnary.cpp b/Engine/Plugins/Experimental/NNE/Source/NNXRuntimeRDG/Private/Hlsl/NNIRuntimeRDGElementWiseUnary.cpp
--- a/Engine/Plugins/Experimental/NNE/Source/NNXRuntimeRDG/Private/Hlsl/NNIRuntimeRDGElementWiseUnary.cpp
+++ b/Engine/Plugins/Experimental/NNE/Source/NNXRuntimeRDG/Private/Hlsl/NNIRuntimeRDGElementWiseUnary.cpp
@@ -49,12 +49,14 @@ namespace UE::NNIRuntimeRDG::Private::Hlsl
 			return true;
 		}
 
-		virtual void Dispatch(FRDGBuilder& GraphBuilder, TConstArrayView<NNX::FTensorRDG> InInputTensors, TConstArrayView<NNX::FTensorRDG> InOutputTensors) override
+		virtual void Dispatch(FRDGBuilder& GraphBuilder, TConstArrayView<NNX::FTensorRDGRef> InInputTensors, TConstArrayView<NNX::FTensorRDGRef> InOutputTensors) override
 		{
-			FRDGBufferSRVRef InputSRV = GraphBuilder.CreateSRV(FRDGBufferSRVDesc(InInputTensors[0].GetBuffer(), PF_R32_FLOAT));
-			FRDGBufferUAVRef OutputUAV = GraphBuilder.CreateUAV(FRDGBufferUAVDesc(InOutputTensors[0].GetBuffer(), PF_R32_FLOAT));
+			check(InInputTensors[0] != nullptr);
+			check(InOutputTensors[0] != nullptr);
+			FRDGBufferSRVRef InputSRV = GraphBuilder.CreateSRV(FRDGBufferSRVDesc(InInputTensors[0]->GetBuffer(), PF_R32_FLOAT));
+			FRDGBufferUAVRef OutputUAV = GraphBuilder.CreateUAV(FRDGBufferUAVDesc(InOutputTensors[0]->GetBuffer(), PF_R32_FLOAT));
 		
-			int32 NumElements = InOutputTensors[0].GetVolume();
+			int32 NumElements = InOutputTensors[0]->GetVolume();
 			FIntVector ThreadGroupCount = NNX::ComputeElementWiseThreadGroups(NumElements, FElementWiseUnaryConstants::NUM_GROUP_THREADS);
 
 			// Set parameters
diff --git a/Engine/Plugins/Experimental/NNE/Source/NNXRuntimeRDG/Private/Hlsl/NNIRuntimeRDGElementWiseVariadic.cpp b/Engine/Plugins/Experimental/NNE/Source/NNXRuntimeRDG/Private/Hlsl/NNIRuntimeRDGElementWiseVariadic.cpp
--- a/Engine/Plugins/Experimental/NNE/Source/NNXRuntimeRDG/Private/Hlsl/NNIRuntimeRDGElementWiseVariadic.cpp
+++ b/Engine/Plugins/Experimental/NNE/Source/NNXRuntimeRDG/Private/Hlsl/NNIRuntimeRDGElementWiseVariadic.cpp
@@ -12,7 +12,7 @@ namespace UE::NNIRuntimeRDG::Private::Hlsl
 	using FElementWiseVariadicConstants = UE::NNEHlslShaders::Internal::FElementWiseVariadicConstants;
 
 	void AddOneVariadicOpPass(FRDGBuilder& GraphBuilder, 
-		TConstArrayView<NNX::FTensorRDG> InputTensors,
+		TConstArrayView<NNX::FTensorRDGRef> InputTensors,
 		const NNX::FTensorRDG& OutputTensor,
 		bool OutputAsInput,
 		EMLElementWiseVariadicOperatorType OpType,
@@ -28,7 +28,8 @@ namespace UE::NNIRuntimeRDG::Private::Hlsl
 
 		for (int32 i = 0; i < InputTensors.Num(); ++i)
 		{
-			InputsSRV[i] = GraphBuilder.CreateSRV(FRDGBufferSRVDesc(InputTensors[i].GetBuffer(), PF_R32_FLOAT));
+			check(InputTensors[i] != nullptr);
+			InputsSRV[i] = GraphBuilder.CreateSRV(FRDGBufferSRVDesc(InputTensors[i]->GetBuffer(), PF_R32_FLOAT));
 		}
 
 		// Set parameters
@@ -40,18 +41,18 @@ namespace UE::NNIRuntimeRDG::Private::Hlsl
 		Params->Input2 = InputsSRV[2];
 		Params->Input3 = InputsSRV[3];
 		Params->Output = OutputUAV;
-		FillTensorStrideForBroadcastShaderParameters(InputTensors[0], OutputTensor.GetShape().Rank(), Params->InputTensorInfo, 0);
+		FillTensorStrideForBroadcastShaderParameters(*InputTensors[0], OutputTensor.GetShape().Rank(), Params->InputTensorInfo, 0);
 		if (InputTensors.Num() >= 2)
 		{
-			FillTensorStrideForBroadcastShaderParameters(InputTensors[1], OutputTensor.GetShape().Rank(), Params->InputTensorInfo, 1);
+			FillTensorStrideForBroadcastShaderParameters(*InputTensors[1], OutputTensor.GetShape().Rank(), Params->InputTensorInfo, 1);
 		}
 		if (InputTensors.Num() >= 3)
 		{
-			FillTensorStrideForBroadcastShaderParameters(InputTensors[2], OutputTensor.GetShape().Rank(), Params->InputTensorInfo, 2);
+			FillTensorStrideForBroadcastShaderParameters(*InputTensors[2], OutputTensor.GetShape().Rank(), Params->InputTensorInfo, 2);
 		}
 		if (InputTensors.Num() >= 4)
 		{
-			FillTensorStrideForBroadcastShaderParameters(InputTensors[3], OutputTensor.GetShape().Rank(), Params->InputTensorInfo, 3);
+			FillTensorStrideForBroadcastShaderParameters(*InputTensors[3], OutputTensor.GetShape().Rank(), Params->InputTensorInfo, 3);
 		}
 		FillTensorStrideShaderParameters(OutputTensor, Params->OutputTensorInfo, 0);
 		Params->Num = OutputTensor.GetVolume();
@@ -136,15 +137,16 @@ namespace UE::NNIRuntimeRDG::Private::Hlsl
 			return true;
 		}
 
-		virtual void Dispatch(FRDGBuilder& GraphBuilder, TConstArrayView<NNX::FTensorRDG> InInputTensors, TConstArrayView<NNX::FTensorRDG> InOutputTensors) override
+		virtual void Dispatch(FRDGBuilder& GraphBuilder, TConstArrayView<NNX::FTensorRDGRef> InInputTensors, TConstArrayView<NNX::FTensorRDGRef> InOutputTensors) override
 		{
 			check(InInputTensors.Num() > 0);
 			check(InOutputTensors.Num() == 1);
+			check(InOutputTensors[0] != nullptr);
 
 			RDG_EVENT_SCOPE(GraphBuilder, "NNI.Operator.Hlsl.ElementWise.Variadic");
 			RDG_GPU_STAT_SCOPE(GraphBuilder, FNNIOperatorElementWiseVariadic);
 
-			NNX::FTensorRDG PassInputTensors[FElementWiseVariadicConstants::MAX_NUM_INPUT];
+			NNX::FTensorRDGRef PassInputTensors[FElementWiseVariadicConstants::MAX_NUM_INPUT];
 			for (int32 InputOffset = 0; InputOffset < InInputTensors.Num(); InputOffset += FElementWiseVariadicConstants::MAX_NUM_INPUT)
 			{
 				uint32 NumInputLeftToHandle = InInputTensors.Num() - InputOffset;
@@ -155,6 +157,7 @@ namespace UE::NNIRuntimeRDG::Private::Hlsl
 
 				for (uint32 i = 0; i < NumInputForPass; ++i)
 				{
+					check(InInputTensors[InputOffset + i] != nullptr);
 					PassInputTensors[i] = InInputTensors[InputOffset + i];
 				}
 
@@ -165,7 +168,7 @@ namespace UE::NNIRuntimeRDG::Private::Hlsl
 			
 				AddOneVariadicOpPass(GraphBuilder,
 					MakeArrayView(PassInputTensors, NumInputForPass),
-					InOutputTensors[0],
+					*InOutputTensors[0],
 					!bIsFirstPass,
 					OpType, Scale);
 			}
diff --git a/Engine/Plugins/Experimental/NNE/Source/NNXRuntimeRDG/Private/Hlsl/NNIRuntimeRDGGather.cpp b/Engine/Plugins/Experimental/NNE/Source/NNXRuntimeRDG/Private/Hlsl/NNIRuntimeRDGGather.cpp
--- a/Engine/Plugins/Experimental/NNE/Source/NNXRuntimeRDG/Private/Hlsl/NNIRuntimeRDGGather.cpp
+++ b/Engine/Plugins/Experimental/NNE/Source/NNXRuntimeRDG/Private/Hlsl/NNIRuntimeRDGGather.cpp
@@ -81,20 +81,23 @@ namespace UE::NNIRuntimeRDG::Private::Hlsl
 			return true;
 		}
 
-		virtual void Dispatch(FRDGBuilder& GraphBuilder, TConstArrayView<NNX::FTensorRDG> InputTensors, TConstArrayView<NNX::FTensorRDG> OutputTensors) override
+		virtual void Dispatch(FRDGBuilder& GraphBuilder, TConstArrayView<NNX::FTensorRDGRef> InputTensors, TConstArrayView<NNX::FTensorRDGRef> OutputTensors) override
 		{
 			using namespace UE::NNEHlslShaders::Internal;
 
 			check(InputTensors.Num() == 2)
 			check(OutputTensors.Num() == 1)
-			check(OutputTensors[0].GetShape().Rank() <= FGatherConstants::MAX_NUM_DIMENSIONS)
-			check(InputTensors[0].GetShape().Rank() > 0)
-			check(InputTensors[1].GetShape().Rank() > 0)
-			check(InputTensors[0].GetShape().Rank() + (InputTensors[1].GetShape().Rank() - 1) <= FGatherConstants::MAX_NUM_DIMENSIONS)
-
-			const NNX::FTensorRDG& Data = InputTensors[0];
-			const NNX::FTensorRDG& Indices = InputTensors[1];
-			const NNX::FTensorRDG& Output = OutputTensors[0];
+			check(InputTensors[0] != nullptr);
+			check(InputTensors[1] != nullptr);
+			check(OutputTensors[0] != nullptr);
+			check(OutputTensors[0]->GetShape().Rank() <= FGatherConstants::MAX_NUM_DIMENSIONS)
+			check(InputTensors[0]->GetShape().Rank() > 0)
+			check(InputTensors[1]->GetShape().Rank() > 0)
+			check(InputTensors[0]->GetShape().Rank() + (InputTensors[1]->GetShape().Rank() - 1) <= FGatherConstants::MAX_NUM_DIMENSIONS)
+
+			const NNX::FTensorRDG& Data = *InputTensors[0];
+			const NNX::FTensorRDG& Indices = *InputTensors[1];
+			const NNX::FTensorRDG& Output = *OutputTensors[0];
 
 			// Set parameters
 			TGatherCS::FParameters* Parameters = GraphBuilder.AllocParameters<TGatherCS::FParameters>();
diff --git a/Engine/Plugins/Experimental/NNE/Source/NNXRuntimeRDG/Private/Hlsl/NNIRuntimeRDGGemm.cpp b/Engine/Plugins/Experimental/NNE/Source/NNXRuntimeRDG/Private/Hlsl/NNIRuntimeRDGGemm.cpp
--- a/Engine/Plugins/Experimental/NNE/Source/NNXRuntimeRDG/Private/Hlsl/NNIRuntimeRDGGemm.cpp
+++ b/Engine/Plugins/Experimental/NNE/Source/NNXRuntimeRDG/Private/Hlsl/NNIRuntimeRDGGemm.cpp
@@ -99,7 +99,7 @@ namespace UE::NNIRuntimeRDG::Private::Hlsl
 			return true;
 		}
 
-		virtual void Dispatch(FRDGBuilder& GraphBuilder, TConstArrayView<NNX::FTensorRDG> InputTensors, TConstArrayView<NNX::FTensorRDG> OutputTensors) override
+		virtual void Dispatch(FRDGBuilder& GraphBuilder, TConstArrayView<NNX::FTensorRDGRef> InputTensors, TConstArrayView<NNX::FTensorRDGRef> OutputTensors) override
 		{
 			using namespace UE::NNEHlslShaders::Internal;
 
@@ -109,15 +109,18 @@ namespace UE::NNIRuntimeRDG::Private::Hlsl
 
 			check(InputTensors.Num() >= 2 && InputTensors.Num() <= 3);
 			check(OutputTensors.Num() == 1);
-
-			const NNX::FTensorRDG& InputA = InputTensors[0];
-			const NNX::FTensorRDG& InputB = InputTensors[1];
-			const NNX::FTensorRDG& Output = OutputTensors[0];
+			check(InputTensors[0] != nullptr);
+			check(InputTensors[1] != nullptr);
+			check(OutputTensors[0] != nullptr);
+			const NNX::FTensorRDG& InputA = *InputTensors[0];
+			const NNX::FTensorRDG& InputB = *InputTensors[1];
+			const NNX::FTensorRDG& Output = *OutputTensors[0];
 			const NNX::FTensorRDG* InputC = nullptr;
 
 			if (InputTensors.Num() == 3)
 			{
-				InputC = &(InputTensors[2]);
+				check(InputTensors[2] != nullptr);
+				InputC = InputTensors[2];
 			}
 			
 			// Set parameters
diff --git a/Engine/Plugins/Experimental/NNE/Source/NNXRuntimeRDG/Private/Hlsl/NNIRuntimeRDGInstanceNormalization.cpp b/Engine/Plugins/Experimental/NNE/Source/NNXRuntimeRDG/Private/Hlsl/NNIRuntimeRDGInstanceNormalization.cpp
--- a/Engine/Plugins/Experimental/NNE/Source/NNXRuntimeRDG/Private/Hlsl/NNIRuntimeRDGInstanceNormalization.cpp
+++ b/Engine/Plugins/Experimental/NNE/Source/NNXRuntimeRDG/Private/Hlsl/NNIRuntimeRDGInstanceNormalization.cpp
@@ -97,17 +97,21 @@ namespace UE::NNIRuntimeRDG::Private::Hlsl
 			return true;
 		}
 
-		virtual void Dispatch(FRDGBuilder& GraphBuilder, TConstArrayView<NNX::FTensorRDG> InputTensors, TConstArrayView<NNX::FTensorRDG> OutputTensors) override
+		virtual void Dispatch(FRDGBuilder& GraphBuilder, TConstArrayView<NNX::FTensorRDGRef> InputTensors, TConstArrayView<NNX::FTensorRDGRef> OutputTensors) override
 		{
 			using namespace UE::NNEHlslShaders::Internal;
 
 			check(InputTensors.Num() == 3);
 			check(OutputTensors.Num() == 1);
-
-			const NNX::FTensorRDG& Input = InputTensors[0];
-			const NNX::FTensorRDG& Scale = InputTensors[1];
-			const NNX::FTensorRDG& Bias = InputTensors[2];
-			const NNX::FTensorRDG& Output = OutputTensors[0];
+			check(InputTensors[0] != nullptr);
+			check(InputTensors[1] != nullptr);
+			check(InputTensors[2] != nullptr);
+			check(OutputTensors[0] != nullptr);
+
+			const NNX::FTensorRDG& Input = *InputTensors[0];
+			const NNX::FTensorRDG& Scale = *InputTensors[1];
+			const NNX::FTensorRDG& Bias = *InputTensors[2];
+			const NNX::FTensorRDG& Output = *OutputTensors[0];
 
 			// Set parameters
 			TInstanceNormalizationCS::FParameters* Parameters = GraphBuilder.AllocParameters<TInstanceNormalizationCS::FParameters>();
diff --git a/Engine/Plugins/Experimental/NNE/Source/NNXRuntimeRDG/Private/Hlsl/NNIRuntimeRDGMatMul.cpp b/Engine/Plugins/Experimental/NNE/Source/NNXRuntimeRDG/Private/Hlsl/NNIRuntimeRDGMatMul.cpp
--- a/Engine/Plugins/Experimental/NNE/Source/NNXRuntimeRDG/Private/Hlsl/NNIRuntimeRDGMatMul.cpp
+++ b/Engine/Plugins/Experimental/NNE/Source/NNXRuntimeRDG/Private/Hlsl/NNIRuntimeRDGMatMul.cpp
@@ -91,16 +91,18 @@ namespace UE::NNIRuntimeRDG::Private::Hlsl
 			return true;
 		}
 
-		virtual void Dispatch(FRDGBuilder& GraphBuilder, TConstArrayView<NNX::FTensorRDG> InputTensors, TConstArrayView<NNX::FTensorRDG> OutputTensors) override
+		virtual void Dispatch(FRDGBuilder& GraphBuilder, TConstArrayView<NNX::FTensorRDGRef> InputTensors, TConstArrayView<NNX::FTensorRDGRef> OutputTensors) override
 		{
 			using namespace UE::NNEHlslShaders::Internal;
 
 			check(InputTensors.Num() == 2);
 			check(OutputTensors.Num() == 1);
-
-			const NNX::FTensorRDG& InputA = InputTensors[0];
-			const NNX::FTensorRDG& InputB = InputTensors[1];
-			const NNX::FTensorRDG& Output = OutputTensors[0];
+			check(InputTensors[0] != nullptr);
+			check(InputTensors[1] != nullptr);
+			check(OutputTensors[0] != nullptr);
+			const NNX::FTensorRDG& InputA = *InputTensors[0];
+			const NNX::FTensorRDG& InputB = *InputTensors[1];
+			const NNX::FTensorRDG& Output = *OutputTensors[0];
 
 			const EGemmAlgorithm Algorithm = EGemmAlgorithm::Simple32x32;
 
diff --git a/Engine/Plugins/Experimental/NNE/Source/NNXRuntimeRDG/Private/NNXRuntimeDML.cpp b/Engine/Plugins/Experimental/NNE/Source/NNXRuntimeRDG/Private/NNXRuntimeDML.cpp
--- a/Engine/Plugins/Experimental/NNE/Source/NNXRuntimeRDG/Private/NNXRuntimeDML.cpp
+++ b/Engine/Plugins/Experimental/NNE/Source/NNXRuntimeRDG/Private/NNXRuntimeDML.cpp
@@ -334,7 +334,7 @@ class FMLOperatorDml : public FMLOperatorRDG
 public:
 
 	virtual bool Initialize(FDeviceContextDml* DevCtx, TArrayView<const FTensor> InputTensors, TArrayView<const FTensor> OutputTensors, const UE::NNECore::FAttributeMap& Attributes) = 0;
-	virtual void Dispatch(FRDGBuilder& GraphBuilder, TConstArrayView<FTensorRDG> InInputTensors, TConstArrayView<FTensorRDG> InOutputTensors) = 0;
+	virtual void Dispatch(FRDGBuilder& GraphBuilder, TConstArrayView<FTensorRDGRef> InInputTensors, TConstArrayView<FTensorRDGRef> InOutputTensors) = 0;
 
 protected:
 
@@ -646,14 +646,16 @@ public:
 	//
 	//
 	//
-	virtual void Dispatch(FRDGBuilder& GraphBuilder, TConstArrayView<FTensorRDG> InInputTensors, TConstArrayView<FTensorRDG> InOutputTensors) override
+	virtual void Dispatch(FRDGBuilder& GraphBuilder, TConstArrayView<FTensorRDGRef> InInputTensors, TConstArrayView<FTensorRDGRef> InOutputTensors) override
 	{
 		ID3D12DynamicRHI* DynamicRHI = GetID3D12DynamicRHI();
 
 		FMLElementWiseUnaryParameters* Params = GraphBuilder.AllocParameters<FMLElementWiseUnaryParameters>();
 
-		Params->Input = InInputTensors[0].GetBuffer();
-		Params->Output = InOutputTensors[0].GetBuffer();
+		check(InInputTensors[0] != nullptr);
+		check(InOutputTensors[0] != nullptr);
+		Params->Input = InInputTensors[0]->GetBuffer();
+		Params->Output = InOutputTensors[0]->GetBuffer();
 		Params->Alpha = Alpha;
 		Params->Beta = Beta;
 		Params->Gamma = Gamma;
@@ -884,15 +886,18 @@ public:
 	//
 	//
 	//
-	virtual void Dispatch(FRDGBuilder& GraphBuilder, TConstArrayView<FTensorRDG> InInputTensors, TConstArrayView<FTensorRDG> InOutputTensors) override
+	virtual void Dispatch(FRDGBuilder& GraphBuilder, TConstArrayView<FTensorRDGRef> InInputTensors, TConstArrayView<FTensorRDGRef> InOutputTensors) override
 	{
 		ID3D12DynamicRHI* DynamicRHI = GetID3D12DynamicRHI();
 
 		FMLElementWiseBinaryParameters* Params = GraphBuilder.AllocParameters<FMLElementWiseBinaryParameters>();
 
-		Params->LHSInput = InInputTensors[0].GetBuffer();
-		Params->RHSInput = InInputTensors[1].GetBuffer();
-		Params->Output = InOutputTensors[0].GetBuffer();
+		check(InInputTensors[0] != nullptr);
+		check(InInputTensors[1] != nullptr);
+		check(InOutputTensors[0] != nullptr);
+		Params->LHSInput = InInputTensors[0]->GetBuffer();
+		Params->RHSInput = InInputTensors[1]->GetBuffer();
+		Params->Output = InOutputTensors[0]->GetBuffer();
 		Params->Num = Num;
 
 		GraphBuilder.AddPass(
@@ -1067,7 +1072,7 @@ public:
 	//
 	//
 	//
-	virtual void Dispatch(FRDGBuilder& GraphBuilder, TConstArrayView<FTensorRDG> InInputTensors, TConstArrayView<FTensorRDG> InOutputTensors) override
+	virtual void Dispatch(FRDGBuilder& GraphBuilder, TConstArrayView<FTensorRDGRef> InInputTensors, TConstArrayView<FTensorRDGRef> InOutputTensors) override
 	{
 		ID3D12DynamicRHI* DynamicRHI = GetID3D12DynamicRHI();
 
@@ -1076,10 +1081,14 @@ public:
 
 		FMLGemmParameters* Params = GraphBuilder.AllocParameters<FMLGemmParameters>();
 
-		Params->A = InInputTensors[0].GetBuffer();
-		Params->B = InInputTensors[1].GetBuffer();
-		Params->C = bIsUsingBias ? InInputTensors[2].GetBuffer() : nullptr;
-		Params->Y = InOutputTensors[0].GetBuffer();
+		check(InInputTensors[0] != nullptr);
+		check(InInputTensors[1] != nullptr);
+		check(!bIsUsingBias || InInputTensors[2] != nullptr);
+		check(InOutputTensors[0] != nullptr);
+		Params->A = InInputTensors[0]->GetBuffer();
+		Params->B = InInputTensors[1]->GetBuffer();
+		Params->C = bIsUsingBias ? InInputTensors[2]->GetBuffer() : nullptr;
+		Params->Y = InOutputTensors[0]->GetBuffer();
 
 		GraphBuilder.AddPass(
 			RDG_EVENT_NAME("FMLGemmDml_Transition"),
@@ -1545,10 +1554,10 @@ void FMLInferenceModelDml::AddDispatchOps_RenderThread(FRDGBuilder& GraphBuilder
 	check(AllTensorRDGs.Num() == AllShapes.Num());
 
 	static constexpr int32 MaxExpectedInput = 10;
-	TArray<FTensorRDG, TInlineAllocator<MaxExpectedInput>> InputTensors;
+	TArray<FTensorRDGRef, TInlineAllocator<MaxExpectedInput>> InputTensors;
 
 	static constexpr int32 MaxExpectedOutput = 2;
-	TArray<FTensorRDG, TInlineAllocator<MaxExpectedOutput>> OutputTensors;
+	TArray<FTensorRDGRef, TInlineAllocator<MaxExpectedOutput>> OutputTensors;
 
 	// Add passes for all operators
 	for (int32 Idx = 0; Idx < Operators.Num(); ++Idx)
diff --git a/Engine/Plugins/Experimental/NNE/Source/NNXRuntimeRDG/Private/NNXRuntimeHLSL.cpp b/Engine/Plugins/Experimental/NNE/Source/NNXRuntimeRDG/Private/NNXRuntimeHLSL.cpp
--- a/Engine/Plugins/Experimental/NNE/Source/NNXRuntimeRDG/Private/NNXRuntimeHLSL.cpp
+++ b/Engine/Plugins/Experimental/NNE/Source/NNXRuntimeRDG/Private/NNXRuntimeHLSL.cpp
@@ -164,6 +164,15 @@ int FMLInferenceModelHlsl::RunShapeInference()
 		checkCode(AllInitializedShapes[InputTensorModelIndex] = true);
 	}
 
+	//Prime shape inference with model weights
+	for (int32 i = 0; i < WeightTensorIndices.Num(); ++i)
+	{
+		const int32 WeightTensorModelIndex = WeightTensorIndices[i];
+
+		AllShapes[WeightTensorModelIndex] = WeightTensorRDGs[i].GetShape();
+		checkCode(AllInitializedShapes[WeightTensorModelIndex] = true);
+	}
+
 	// Run shape inference on all operators
 	for (int32 Idx = 0; Idx < Operators.Num(); ++Idx)
 	{
@@ -237,10 +246,10 @@ void FMLInferenceModelHlsl::AddDispatchOps_RenderThread(FRDGBuilder& GraphBuilde
 	check(AllTensorRDGs.Num() == AllShapes.Num());
 
 	static constexpr int32 MaxExpectedInput = 10;
-	TArray<FTensorRDG, TInlineAllocator<MaxExpectedInput>> InputTensors;
+	TArray<FTensorRDGRef, TInlineAllocator<MaxExpectedInput>> InputTensors;
 
 	static constexpr int32 MaxExpectedOutput = 2;
-	TArray<FTensorRDG, TInlineAllocator<MaxExpectedOutput>> OutputTensors;
+	TArray<FTensorRDGRef, TInlineAllocator<MaxExpectedOutput>> OutputTensors;
 
 	// Add passes for all operators
 	for (int32 Idx = 0; Idx < Operators.Num(); ++Idx)
diff --git a/Engine/Plugins/Experimental/NNE/Source/NNXRuntimeRDG/Private/NNXRuntimeHLSLOp.h b/Engine/Plugins/Experimental/NNE/Source/NNXRuntimeRDG/Private/NNXRuntimeHLSLOp.h
--- a/Engine/Plugins/Experimental/NNE/Source/NNXRuntimeRDG/Private/NNXRuntimeHLSLOp.h
+++ b/Engine/Plugins/Experimental/NNE/Source/NNXRuntimeRDG/Private/NNXRuntimeHLSLOp.h
@@ -14,7 +14,7 @@ namespace NNX
 	{
 		virtual ~FMLOperatorHlsl() = default;
 		virtual bool Initialize(TConstArrayView<FTensorDesc> InputTensorDescs, TConstArrayView<FTensorDesc> OutputTensorDescs, const UE::NNECore::FAttributeMap& Attributes) = 0;
-		virtual void Dispatch(FRDGBuilder& GraphBuilder, TConstArrayView<FTensorRDG> InInputTensors, TConstArrayView<FTensorRDG> InOutputTensors) = 0;
+		virtual void Dispatch(FRDGBuilder& GraphBuilder, TConstArrayView<FTensorRDGRef> InInputTensors, TConstArrayView<FTensorRDGRef> InOutputTensors) = 0;
 		virtual int ComputeOutputShape(TConstArrayView<FTensorShape> InputShapes, TArray<FTensorShape>& OutputShapes) const { return -1; };
 	};
 
diff --git a/Engine/Plugins/Experimental/NNE/Source/NNXRuntimeRDG/Private/NNXRuntimeRDG.cpp b/Engine/Plugins/Experimental/NNE/Source/NNXRuntimeRDG/Private/NNXRuntimeRDG.cpp
--- a/Engine/Plugins/Experimental/NNE/Source/NNXRuntimeRDG/Private/NNXRuntimeRDG.cpp
+++ b/Engine/Plugins/Experimental/NNE/Source/NNXRuntimeRDG/Private/NNXRuntimeRDG.cpp
@@ -227,6 +227,7 @@ bool FMLInferenceModelRDG::LoadModel(TConstArrayView<uint8> ModelData, FMLRuntim
 	// Data for RDG
 	AllSymbolicTensorDescs.Empty();
 	IntermediateTensorIndices.Empty();
+	WeightTensorIndices.Empty();
 	InputTensorIndices.Empty();
 	OutputTensorIndices.Empty();
 	OperatorInputTensorIndices.Empty();
@@ -255,6 +256,29 @@ bool FMLInferenceModelRDG::LoadModel(TConstArrayView<uint8> ModelData, FMLRuntim
 		{
 			IntermediateTensorIndices.Emplace(Idx);
 		}
+		else if (FormatTensorDesc.Type == EMLFormatTensorType::Initializer)
+		{
+			WeightTensorIndices.Emplace(Idx);
+			if (!SymbolicTensor.IsConcrete())
+			{
+				UE_LOG(LogNNX, Error, TEXT("Weight tensor %s should have a concrete shape"), *SymbolicTensor.GetName());
+				return false;
+			}
+			
+			const FTensorShape TensorShape = FTensorShape::MakeFromSymbolic(SymbolicTensor.GetShape());
+			FTensorRDG& WeightRDG = WeightTensorRDGs.Emplace_GetRef(FTensorRDG::Make(SymbolicTensor, TensorShape, nullptr));
+			
+			if (WeightRDG.GetDataSize() != FormatTensorDesc.DataSize)
+			{
+				UE_LOG(LogNNX, Error, TEXT("Weight %s has incorrect size. Expected %d bytes, got %d"), *SymbolicTensor.GetName(), FormatTensorDesc.DataSize, WeightRDG.GetDataSize());
+				return false;
+			}
+
+			const uint8* DataPtr = Format.TensorData.GetData() + FormatTensorDesc.DataOffset;
+			TConstArrayView<uint8> DataView = MakeArrayView(DataPtr, FormatTensorDesc.DataSize);
+
+			WeightRDG.SetPreparedData(DataView);
+		}
 		checkf(FormatTensorDesc.Type != EMLFormatTensorType::None, TEXT("Unsupported tensor type None"));
 	}
 
@@ -351,9 +375,10 @@ int FMLInferenceModelRDG::SetInputTensorShapes(TConstArrayView<FTensorShape> InI
 		OutputTensorShapes.Emplace(AllShapes[OutputIndices]);
 	}
 
-	check(InputTensorIndices.Num() + OutputTensorIndices.Num() + IntermediateTensorIndices.Num() == AllShapes.Num());
+	check(InputTensorIndices.Num() + OutputTensorIndices.Num() + WeightTensorIndices.Num() + IntermediateTensorIndices.Num() == AllShapes.Num());
 	check(InputTensorShapes.Num() == InputSymbolicTensors.Num());
 	check(OutputTensorShapes.Num() == OutputSymbolicTensors.Num());
+	check(WeightTensorIndices.Num() == WeightTensorRDGs.Num());
 	check(AllShapes.Num() == AllSymbolicTensorDescs.Num());
 	
 	return 0;
@@ -385,7 +410,7 @@ int FMLInferenceModelRDG::EnqueueRDG(FRDGBuilder& RDGBuilder, TConstArrayView<FM
 	}
 
 	// Process input binding, and if required, allocate RDG buffers
-	FTensorRDGArray	InputTensorRDGs;
+	InputTensorRDGs.Empty();
 	FIntArray		InputUploadIndices;
 
 	Res = SetTensors(RDGBuilder, InputTensorRDGs, InputUploadIndices, InInputBindings, InputSymbolicTensors, InputTensorShapes);
@@ -396,7 +421,7 @@ int FMLInferenceModelRDG::EnqueueRDG(FRDGBuilder& RDGBuilder, TConstArrayView<FM
 	}
 
 	// Process output tensors, and if required, allocate RDG buffers
-	FTensorRDGArray	OutputTensorRDGs;
+	OutputTensorRDGs.Empty();
 	FIntArray		OutputReadbackIndices;
 
 	Res = SetTensors(RDGBuilder, OutputTensorRDGs, OutputReadbackIndices, InOutputBindings, OutputSymbolicTensors, OutputTensorShapes);
@@ -411,33 +436,42 @@ int FMLInferenceModelRDG::EnqueueRDG(FRDGBuilder& RDGBuilder, TConstArrayView<FM
 
 	AllTensorRDGs.Empty();
 	AllTensorRDGs.SetNum(AllShapes.Num());
-	checkCode( for (FTensorRDG& TensorRDG : AllTensorRDGs) { TensorRDG.SetBuffer(nullptr); } );
 	
-	//Create intermediate tensors bindings
-	for (int32 Idx : IntermediateTensorIndices)
+	//Create intermediate tensors
+	IntermediateTensorRDGs.Empty();
+	for (int32 i = 0; i < IntermediateTensorIndices.Num(); ++i)
 	{
+		const int32 Idx = IntermediateTensorIndices[i];
 		const FTensorDesc& TensorDesc = AllSymbolicTensorDescs[Idx];
 		const FTensorShape& TensorShape = AllShapes[Idx];
-		FTensorRDG TensorRDG = FTensorRDG::Make(TensorDesc, TensorShape, nullptr);
+		FTensorRDG& TensorRDG = IntermediateTensorRDGs.Emplace_GetRef(FTensorRDG::Make(TensorDesc, TensorShape, nullptr));
 		const FRDGBufferDesc BufferDesc = CreateRDGBufferDescForTensorRDG(TensorRDG);
 		const FRDGBufferRef TensorBuffer = RDGBuilder.CreateBuffer(BufferDesc, *TensorRDG.GetName(), ERDGBufferFlags::None);
 		
 		TensorRDG.SetBuffer(TensorBuffer);
-		AllTensorRDGs[Idx] = TensorRDG;
+		AllTensorRDGs[Idx] = &IntermediateTensorRDGs[i];
 	}
 
-	//Insert input tensors bindings
+	//Insert input tensors
 	for (int32 i = 0; i < InputTensorIndices.Num(); ++i)
 	{
-		AllTensorRDGs[InputTensorIndices[i]] = InputTensorRDGs[i];
+		AllTensorRDGs[InputTensorIndices[i]] = &InputTensorRDGs[i];
 	}
 
-	//Insert output tensors bindings
+	//Insert output tensors
 	for (int32 i = 0; i < OutputTensorIndices.Num(); ++i)
 	{
-		AllTensorRDGs[OutputTensorIndices[i]] = OutputTensorRDGs[i];
+		AllTensorRDGs[OutputTensorIndices[i]] = &OutputTensorRDGs[i];
+	}
+
+	//For now weights tensors are not uploaded to GPU thus GetBuffer will return nullptr for them.
+	checkCode(for (const FTensorRDG* TensorRDG : AllTensorRDGs) { if (TensorRDG != nullptr) { check(TensorRDG->GetBuffer() != nullptr); } });
+
+	//Insert weights tensors
+	for (int32 i = 0; i < WeightTensorIndices.Num(); ++i)
+	{
+		AllTensorRDGs[WeightTensorIndices[i]] = &WeightTensorRDGs[i];
 	}
-	checkCode(for (const FTensorRDG& TensorRDG : AllTensorRDGs) { check(TensorRDG.GetBuffer() != nullptr); });
 	
 	// We can now dispatch operators
 	AddDispatchOps_RenderThread(RDGBuilder);
diff --git a/Engine/Plugins/Experimental/NNE/Source/NNXRuntimeRDG/Private/NNXRuntimeRDG.h b/Engine/Plugins/Experimental/NNE/Source/NNXRuntimeRDG/Private/NNXRuntimeRDG.h
--- a/Engine/Plugins/Experimental/NNE/Source/NNXRuntimeRDG/Private/NNXRuntimeRDG.h
+++ b/Engine/Plugins/Experimental/NNE/Source/NNXRuntimeRDG/Private/NNXRuntimeRDG.h
@@ -65,7 +65,9 @@ public:
 	FRDGBufferRef GetBuffer() const { return Buffer; }
 };
 
+using FTensorRDGRef = FTensorRDG*;
 using FTensorRDGArray = TArray<FTensorRDG, TInlineAllocator<16>>;
+using FTensorRDGRefArray = TArray<FTensorRDGRef, TInlineAllocator<64>>;
 using FIntArray = TArray<int32, TInlineAllocator<16>>;
 
 class FMLRuntimeRDG : public IRuntime
@@ -115,17 +117,22 @@ protected:
 
 	//Tensor indices for models
 	TArray<int32>				IntermediateTensorIndices;
+	TArray<int32>				WeightTensorIndices;
 	TArray<int32>				InputTensorIndices;
 	TArray<int32>				OutputTensorIndices;
 	
 	//Tensor indices by operator
 	TArray<TArray<uint32>>		OperatorInputTensorIndices;
 	TArray<TArray<uint32>>		OperatorOutputTensorIndices;
-	
+
 	//RDG Tensors
-	FTensorRDGArray				AllTensorRDGs;
+	FTensorRDGRefArray			AllTensorRDGs;
+	FTensorRDGArray				InputTensorRDGs;
+	FTensorRDGArray				OutputTensorRDGs;
+	FTensorRDGArray				IntermediateTensorRDGs;
+	FTensorRDGArray				WeightTensorRDGs;
 	
-	TArray<FReadbackEntry> 		Readbacks;
+	TArray<FReadbackEntry>		Readbacks;
 	bool						bUseManualTransitions;
 };
 
