NNE - Fix FindPlugin when loading tests from Json

diff --git a/Engine/Plugins/Experimental/NNE/Source/NNXQA/Private/NNXQAParametricTest.cpp b/Engine/Plugins/Experimental/NNE/Source/NNXQA/Private/NNXQAParametricTest.cpp
--- a/Engine/Plugins/Experimental/NNE/Source/NNXQA/Private/NNXQAParametricTest.cpp
+++ b/Engine/Plugins/Experimental/NNE/Source/NNXQA/Private/NNXQAParametricTest.cpp
@@ -35,17 +35,18 @@ namespace Test
 		{
 			TestSetups.Empty();
 
-			TSharedPtr<IPlugin> NNXPlugin = IPluginManager::Get().FindPlugin(TEXT("NNX"));
+			TSharedPtr<IPlugin> NNEPlugin = IPluginManager::Get().FindPlugin(TEXT("NNE"));
 			
-			if (!NNXPlugin.IsValid())
+			if (!NNEPlugin.IsValid())
 			{
+				UE_LOG(LogNNX, Error, TEXT("Unable to find NNE plugin!"));
 				return false;
 			}
 
 			//TODO verify the path can be accessed on standalone build and on consoles
 			//TODO allow to define the tests in more than one json file
-			FString NNXPluginBaseDir = NNXPlugin->GetBaseDir();
-			FString FullPath(NNXPluginBaseDir + TEXT("\\Source\\NNXQA\\Resources\\NNXQATestDesc.json"));
+			FString NNEPluginBaseDir = NNEPlugin->GetBaseDir();
+			FString FullPath(NNEPluginBaseDir + TEXT("\\Source\\NNXQA\\Resources\\NNXQATestDesc.json"));
 
 			TArray<Json::FTestCategory> ModelTestCategories;
 			TArray<Json::FTestCategory> OperatorTestCategories;
@@ -57,9 +58,9 @@ namespace Test
 				return false;
 			}
 
-			const FString NNXBaseTestPath(TEXT("System.Engine.MachineLearning.NNX"));
-			AddTestFromCategory(NNXBaseTestPath + TEXT(".Model."), ModelTestCategories, InputOutputSets, AttributeSets);
-			AddTestFromCategory(NNXBaseTestPath + TEXT(".Operator."), OperatorTestCategories, InputOutputSets, AttributeSets);
+			const FString NNEBaseTestPath(TEXT("System.Engine.MachineLearning.NNE"));
+			AddTestFromCategory(NNEBaseTestPath + TEXT(".Model."), ModelTestCategories, InputOutputSets, AttributeSets);
+			AddTestFromCategory(NNEBaseTestPath + TEXT(".Operator."), OperatorTestCategories, InputOutputSets, AttributeSets);
 
 			return true;
 		}
