NNE constant tensor as input to shape inference

Problem:
* Shape inference need to be able to access model weights for some operator to be able to compute the intermediate buffers sizes.

Solution:
* Refactor shape inference into a model preparation step taking tensors as inputs, this will latter be extensible to further chain of const tensor dependancies for shape inference.

Tests:
* All NNX tests pass under win64 editor

diff --git a/Engine/Plugins/Experimental/NNE/Source/NNXCore/Public/NNXTypes.h b/Engine/Plugins/Experimental/NNE/Source/NNXCore/Public/NNXTypes.h
--- a/Engine/Plugins/Experimental/NNE/Source/NNXCore/Public/NNXTypes.h
+++ b/Engine/Plugins/Experimental/NNE/Source/NNXCore/Public/NNXTypes.h
@@ -60,6 +60,12 @@ template <typename DimType, class FinalType> struct TTensorShapeBase
 	
 	inline bool operator!=(const TTensorShapeBase<DimType, FinalType>& OtherShape) const { return !(*this == OtherShape); }
 
+	inline void operator=(const TTensorShapeBase<DimType, FinalType>& OtherShape)
+	{
+		check(Rank() <= MaxRank);
+		Data = OtherShape.Data;
+	}
+
 	static FinalType Make(TConstArrayView<DimType> Data)
 	{
 		if (Data.Num() > MaxRank)
@@ -246,9 +252,18 @@ public:
 		return PreparedData;
 	}
 
+	void SetShape(const FTensorShape& InShape)
+	{
+		checkf(!HasPreparedData(), TEXT("Shape cannot be changed once data as been set."));
+		check(InShape.Volume() <= TNumericLimits<uint32>::Max());
+		Shape = InShape;
+		Volume = InShape.Volume();
+		DataSize = static_cast<uint64>(GetTensorDataTypeSizeInBytes(DataType) * Volume);
+	}
+
 	void SetPreparedData(TConstArrayView<uint8> Data)
 	{
-		check(Data.Num() == DataSize);
+		checkf(Data.Num() == DataSize, TEXT("Incorrect data size, it should match tensor shape and data type."));
 		PreparedData.Empty();
 		PreparedData.Append(Data);
 	}
@@ -267,17 +282,20 @@ public:
 	{
 		return DataSize;
 	}
-	
+
 	static FTensor Make(const FString& Name, const FTensorShape& Shape, EMLTensorDataType DataType)
 	{
-		FTensor Desc;
-		Desc.Name = Name;
-		Desc.DataType = DataType;
-		Desc.Shape = Shape;
-		Desc.Volume = Shape.Volume();
-		Desc.DataSize = (uint64)GetTensorDataTypeSizeInBytes(DataType) * Desc.Volume;
-		check(Desc.Volume <= TNumericLimits<uint32>::Max());
-		return Desc;
+		FTensor Tensor;
+		Tensor.Name = Name;
+		Tensor.DataType = DataType;
+		Tensor.SetShape(Shape);
+		return Tensor;
+	}
+
+	static FTensor Make(const FTensorDesc& TensorDesc, const FTensorShape& Shape)
+	{
+		check(Shape.IsCompatibleWith(TensorDesc.GetShape()));
+		return Make(TensorDesc.GetName(), Shape, TensorDesc.GetDataType());
 	}
 
 	static FTensor MakeFromSymbolicDesc(const FTensorDesc& TensorDesc)
@@ -286,4 +304,6 @@ public:
 	}
 };
 
+using FTensorRef = FTensor*;
+
 } // namespace NNX
diff --git a/Engine/Plugins/Experimental/NNE/Source/NNXQA/Resources/NNXQATestDesc.json b/Engine/Plugins/Experimental/NNE/Source/NNXQA/Resources/NNXQATestDesc.json
--- a/Engine/Plugins/Experimental/NNE/Source/NNXQA/Resources/NNXQATestDesc.json
+++ b/Engine/Plugins/Experimental/NNE/Source/NNXQA/Resources/NNXQATestDesc.json
@@ -661,7 +661,7 @@
 				{
 					"inputs": [ { "shape": [ 1, 128, 56, 56 ] }, { "shape": [ 128 ] }, { "shape": [ 128 ] } ], "outputs": [ { "shape": [ 1, 128, 56, 56 ] } ],
 					"runtimes": [
-						{ "name": "NNXRuntimeHLSL", "relative_error": 9 },
+						{ "name": "NNXRuntimeHLSL", "relative_error": 10.5 },
 						{ "name": "NNXRuntimeORTCpu", "relative_error": 9 },
 						{ "name": "NNXRuntimeORTDml", "relative_error": 9 }
 					],
diff --git a/Engine/Plugins/Experimental/NNE/Source/NNXRuntimeRDG/Private/Hlsl/NNIRuntimeRDGConv.cpp b/Engine/Plugins/Experimental/NNE/Source/NNXRuntimeRDG/Private/Hlsl/NNIRuntimeRDGConv.cpp
--- a/Engine/Plugins/Experimental/NNE/Source/NNXRuntimeRDG/Private/Hlsl/NNIRuntimeRDGConv.cpp
+++ b/Engine/Plugins/Experimental/NNE/Source/NNXRuntimeRDG/Private/Hlsl/NNIRuntimeRDGConv.cpp
@@ -39,13 +39,13 @@ namespace UE::NNIRuntimeRDG::Private::Hlsl
 
 	public:
 
-		virtual int ComputeOutputShape(TConstArrayView<NNX::FTensorShape> InputShapes, TArray<NNX::FTensorShape>& OutputShapes) const override
+		virtual int PrepareOutputs(TConstArrayView<NNX::FTensorRef> InputTensors, TArrayView<NNX::FTensorRef> OutputTensors) const override
 		{
-			OutputShapes.Empty();
-			check(InputShapes.Num() >= 2 && InputShapes.Num() <= 3);
+			check(InputTensors.Num() >= 2 && InputTensors.Num() <= 3);
+			check(OutputTensors.Num() == 1);
 
-			const NNX::FTensorShape& Input = InputShapes[0];
-			const NNX::FTensorShape& Weights = InputShapes[1];
+			const NNX::FTensorShape& Input = InputTensors[0]->GetShape();
+			const NNX::FTensorShape& Weights = InputTensors[1]->GetShape();
 			NNX::FSymbolicTensorShape OutputShape;
 			
 			OutputShape.Data = UE::NNEHlslShaders::Internal::FConvCS::GetOutputShape(Input.Data, Weights.Data, AutoPad, Dilations, Strides, Pads);
@@ -53,7 +53,7 @@ namespace UE::NNIRuntimeRDG::Private::Hlsl
 			{
 				return -1;
 			}
-			OutputShapes.Emplace(NNX::FTensorShape::MakeFromSymbolic(OutputShape));
+			OutputTensors[0]->SetShape(NNX::FTensorShape::MakeFromSymbolic(OutputShape));
 
 			return 0;
 		};
diff --git a/Engine/Plugins/Experimental/NNE/Source/NNXRuntimeRDG/Private/Hlsl/NNIRuntimeRDGConvTranspose.cpp b/Engine/Plugins/Experimental/NNE/Source/NNXRuntimeRDG/Private/Hlsl/NNIRuntimeRDGConvTranspose.cpp
--- a/Engine/Plugins/Experimental/NNE/Source/NNXRuntimeRDG/Private/Hlsl/NNIRuntimeRDGConvTranspose.cpp
+++ b/Engine/Plugins/Experimental/NNE/Source/NNXRuntimeRDG/Private/Hlsl/NNIRuntimeRDGConvTranspose.cpp
@@ -39,13 +39,13 @@ namespace UE::NNIRuntimeRDG::Private::Hlsl
 
 	public:
 
-		virtual int ComputeOutputShape(TConstArrayView<NNX::FTensorShape> InputShapes, TArray<NNX::FTensorShape>& OutputShapes) const override
+		virtual int PrepareOutputs(TConstArrayView<NNX::FTensorRef> InputTensors, TArrayView<NNX::FTensorRef> OutputTensors) const override
 		{
-			OutputShapes.Empty();
-			check(InputShapes.Num() >= 2 && InputShapes.Num() <= 3);
+			check(InputTensors.Num() >= 2 && InputTensors.Num() <= 3);
+			check(OutputTensors.Num() == 1);
 			
-			const NNX::FTensorShape& Input = InputShapes[0];
-			const NNX::FTensorShape& Weights = InputShapes[1];
+			const NNX::FTensorShape& Input = InputTensors[0]->GetShape();
+			const NNX::FTensorShape& Weights = InputTensors[1]->GetShape();
 			NNX::FSymbolicTensorShape OutputShape;
 			
 			OutputShape.Data = UE::NNEHlslShaders::Internal::FConvTransposeCS::GetOutputShape(Input.Data, Weights.Data, AutoPad, Dilations, Strides, Pads, OutputPadding, Group);
@@ -53,7 +53,7 @@ namespace UE::NNIRuntimeRDG::Private::Hlsl
 			{
 				return -1;
 			}
-			OutputShapes.Emplace(NNX::FTensorShape::MakeFromSymbolic(OutputShape));
+			OutputTensors[0]->SetShape(NNX::FTensorShape::MakeFromSymbolic(OutputShape));
 
 			return 0;
 		};
diff --git a/Engine/Plugins/Experimental/NNE/Source/NNXRuntimeRDG/Private/Hlsl/NNIRuntimeRDGElementWiseBinary.cpp b/Engine/Plugins/Experimental/NNE/Source/NNXRuntimeRDG/Private/Hlsl/NNIRuntimeRDGElementWiseBinary.cpp
--- a/Engine/Plugins/Experimental/NNE/Source/NNXRuntimeRDG/Private/Hlsl/NNIRuntimeRDGElementWiseBinary.cpp
+++ b/Engine/Plugins/Experimental/NNE/Source/NNXRuntimeRDG/Private/Hlsl/NNIRuntimeRDGElementWiseBinary.cpp
@@ -24,13 +24,13 @@ namespace UE::NNIRuntimeRDG::Private::Hlsl
 
 	public:
 
-		virtual int ComputeOutputShape(TConstArrayView<NNX::FTensorShape> InputShapes, TArray<NNX::FTensorShape>& OutputShapes) const override
+		virtual int PrepareOutputs(TConstArrayView<NNX::FTensorRef> InputTensors, TArrayView<NNX::FTensorRef> OutputTensors) const override
 		{
-			OutputShapes.Empty();
-			check(InputShapes.Num() == 2);
+			check(InputTensors.Num() == 2);
+			check(OutputTensors.Num() == 1);
 			
-			const NNX::FTensorShape& LHSInput = InputShapes[0];
-			const NNX::FTensorShape& RHSInput = InputShapes[1];
+			const NNX::FTensorShape& LHSInput = InputTensors[0]->GetShape();
+			const NNX::FTensorShape& RHSInput = InputTensors[1]->GetShape();
 			const int32 OutputRank = FMath::Max(LHSInput.Rank(), RHSInput.Rank());
 			NNX::FTensorShape OutputShape;
 			
@@ -51,7 +51,7 @@ namespace UE::NNIRuntimeRDG::Private::Hlsl
 				OutputShape.Data[OutputRank - 1 - i] = OutputValue;
 			}
 
-			OutputShapes.Emplace(OutputShape);
+			OutputTensors[0]->SetShape(OutputShape);
 			
 			return 0;
 		}
diff --git a/Engine/Plugins/Experimental/NNE/Source/NNXRuntimeRDG/Private/Hlsl/NNIRuntimeRDGElementWiseUnary.cpp b/Engine/Plugins/Experimental/NNE/Source/NNXRuntimeRDG/Private/Hlsl/NNIRuntimeRDGElementWiseUnary.cpp
--- a/Engine/Plugins/Experimental/NNE/Source/NNXRuntimeRDG/Private/Hlsl/NNIRuntimeRDGElementWiseUnary.cpp
+++ b/Engine/Plugins/Experimental/NNE/Source/NNXRuntimeRDG/Private/Hlsl/NNIRuntimeRDGElementWiseUnary.cpp
@@ -31,10 +31,11 @@ namespace UE::NNIRuntimeRDG::Private::Hlsl
 
 	public:
 
-		virtual int ComputeOutputShape(TConstArrayView<NNX::FTensorShape> InputShapes, TArray<NNX::FTensorShape>& OutputShapes) const override
+		virtual int PrepareOutputs(TConstArrayView<NNX::FTensorRef> InputTensors, TArrayView<NNX::FTensorRef> OutputTensors) const override
 		{
-			check(InputShapes.Num() == 1);
-			OutputShapes = InputShapes;
+			check(InputTensors.Num() == 1);
+			check(OutputTensors.Num() == 1);
+			OutputTensors[0]->SetShape(InputTensors[0]->GetShape());
 			return 0;
 		}
 
diff --git a/Engine/Plugins/Experimental/NNE/Source/NNXRuntimeRDG/Private/Hlsl/NNIRuntimeRDGElementWiseVariadic.cpp b/Engine/Plugins/Experimental/NNE/Source/NNXRuntimeRDG/Private/Hlsl/NNIRuntimeRDGElementWiseVariadic.cpp
--- a/Engine/Plugins/Experimental/NNE/Source/NNXRuntimeRDG/Private/Hlsl/NNIRuntimeRDGElementWiseVariadic.cpp
+++ b/Engine/Plugins/Experimental/NNE/Source/NNXRuntimeRDG/Private/Hlsl/NNIRuntimeRDGElementWiseVariadic.cpp
@@ -92,16 +92,16 @@ namespace UE::NNIRuntimeRDG::Private::Hlsl
 
 	public:
 
-		virtual int ComputeOutputShape(TConstArrayView<NNX::FTensorShape> InputShapes, TArray<NNX::FTensorShape>& OutputShapes) const override
+		virtual int PrepareOutputs(TConstArrayView<NNX::FTensorRef> InputTensors, TArrayView<NNX::FTensorRef> OutputTensors) const override
 		{
-			OutputShapes.Empty();
-			check(InputShapes.Num() > 0);
+			check(InputTensors.Num() > 0);
+			check(OutputTensors.Num() == 1);
 
-			const int32 NumInput = InputShapes.Num();
+			const int32 NumInput = InputTensors.Num();
 			int32 OutputRank = 0;
 			for (int32 i = 0; i < NumInput; ++i)
 			{
-				OutputRank = FMath::Max(OutputRank, InputShapes[i].Rank());
+				OutputRank = FMath::Max(OutputRank, InputTensors[i]->GetShape().Rank());
 			}
 
 			NNX::FTensorShape OutputShape;
@@ -113,8 +113,8 @@ namespace UE::NNIRuntimeRDG::Private::Hlsl
 				int32 OutputValue = 1;
 				for (int32 InputIdx = 0; InputIdx < NumInput; ++InputIdx)
 				{
-					int32 InputIndex = InputShapes[InputIdx].Rank() - 1 - i;
-					int32 InputValue = InputIndex >= 0 ? InputShapes[InputIdx].Data[InputIndex] : 1;
+					int32 InputIndex = InputTensors[InputIdx]->GetShape().Rank() - 1 - i;
+					int32 InputValue = InputIndex >= 0 ? InputTensors[InputIdx]->GetShape().Data[InputIndex] : 1;
 					if (InputValue != OutputValue && InputValue != 1 && OutputValue != 1)
 					{
 						UE_LOG(LogNNX, Warning, TEXT("Error while computing shape for element wise variadic op, input shapes are not compatible"));
@@ -125,7 +125,7 @@ namespace UE::NNIRuntimeRDG::Private::Hlsl
 				OutputShape.Data[OutputRank - 1 - i] = OutputValue;
 			}
 
-			OutputShapes.Emplace(OutputShape);
+			OutputTensors[0]->SetShape(OutputShape);
 
 			return 0;
 		}
diff --git a/Engine/Plugins/Experimental/NNE/Source/NNXRuntimeRDG/Private/Hlsl/NNIRuntimeRDGGather.cpp b/Engine/Plugins/Experimental/NNE/Source/NNXRuntimeRDG/Private/Hlsl/NNIRuntimeRDGGather.cpp
--- a/Engine/Plugins/Experimental/NNE/Source/NNXRuntimeRDG/Private/Hlsl/NNIRuntimeRDGGather.cpp
+++ b/Engine/Plugins/Experimental/NNE/Source/NNXRuntimeRDG/Private/Hlsl/NNIRuntimeRDGGather.cpp
@@ -26,10 +26,11 @@ namespace UE::NNIRuntimeRDG::Private::Hlsl
 
 	public:
 
-		virtual int ComputeOutputShape(TConstArrayView<NNX::FTensorShape> InputShapes, TArray<NNX::FTensorShape>& OutputShapes) const override
+		virtual int PrepareOutputs(TConstArrayView<NNX::FTensorRef> InputTensors, TArrayView<NNX::FTensorRef> OutputTensors) const override
 		{
-			OutputShapes.Empty();
-			UE_LOG(LogNNX, Warning, TEXT("Gather does not support variable input shapes at the moment"));
+			check(InputTensors.Num() == 2)
+			check(OutputTensors.Num() == 1)
+			UE_LOG(LogNNX, Warning, TEXT("Gather shape inference is not implemented at the moment"));
 			return -1;
 		};
 		
diff --git a/Engine/Plugins/Experimental/NNE/Source/NNXRuntimeRDG/Private/Hlsl/NNIRuntimeRDGGemm.cpp b/Engine/Plugins/Experimental/NNE/Source/NNXRuntimeRDG/Private/Hlsl/NNIRuntimeRDGGemm.cpp
--- a/Engine/Plugins/Experimental/NNE/Source/NNXRuntimeRDG/Private/Hlsl/NNIRuntimeRDGGemm.cpp
+++ b/Engine/Plugins/Experimental/NNE/Source/NNXRuntimeRDG/Private/Hlsl/NNIRuntimeRDGGemm.cpp
@@ -31,13 +31,13 @@ namespace UE::NNIRuntimeRDG::Private::Hlsl
 
 	public:
 
-		virtual int ComputeOutputShape(TConstArrayView<NNX::FTensorShape> InputShapes, TArray<NNX::FTensorShape>& OutputShapes) const override
+		virtual int PrepareOutputs(TConstArrayView<NNX::FTensorRef> InputTensors, TArrayView<NNX::FTensorRef> OutputTensors) const override
 		{
-			OutputShapes.Empty();
-			check(InputShapes.Num() >= 2 && InputShapes.Num() <= 3);
+			check(InputTensors.Num() >= 2 && InputTensors.Num() <= 3);
+			check(OutputTensors.Num() == 1);
 
-			const NNX::FTensorShape& InputA = InputShapes[0];
-			const NNX::FTensorShape& InputB = InputShapes[1];
+			const NNX::FTensorShape& InputA = InputTensors[0]->GetShape();
+			const NNX::FTensorShape& InputB = InputTensors[1]->GetShape();
 			if (InputA.Rank() != 2 || InputB.Rank() != 2)
 			{
 				return -1;
@@ -49,7 +49,7 @@ namespace UE::NNIRuntimeRDG::Private::Hlsl
 			
 			OutputShape.Data.Emplace(M);
 			OutputShape.Data.Emplace(N);
-			OutputShapes.Emplace(OutputShape);
+			OutputTensors[0]->SetShape(OutputShape);
 			return 0;
 		};
 		
diff --git a/Engine/Plugins/Experimental/NNE/Source/NNXRuntimeRDG/Private/Hlsl/NNIRuntimeRDGInstanceNormalization.cpp b/Engine/Plugins/Experimental/NNE/Source/NNXRuntimeRDG/Private/Hlsl/NNIRuntimeRDGInstanceNormalization.cpp
--- a/Engine/Plugins/Experimental/NNE/Source/NNXRuntimeRDG/Private/Hlsl/NNIRuntimeRDGInstanceNormalization.cpp
+++ b/Engine/Plugins/Experimental/NNE/Source/NNXRuntimeRDG/Private/Hlsl/NNIRuntimeRDGInstanceNormalization.cpp
@@ -66,18 +66,17 @@ namespace UE::NNIRuntimeRDG::Private::Hlsl
 
 	public:
 
-		virtual int ComputeOutputShape(TConstArrayView<NNX::FTensorShape> InputShapes, TArray<NNX::FTensorShape>& OutputShapes) const override
+		virtual int PrepareOutputs(TConstArrayView<NNX::FTensorRef> InputTensors, TArrayView<NNX::FTensorRef> OutputTensors) const override
 		{
-			check(InputShapes.Num() == 3);
-
-			OutputShapes.Empty(); // TODO REMOVE caller should make sure empty array is passed!
+			check(InputTensors.Num() == 3);
+			check(OutputTensors.Num() == 1);
 
-			if (const int Res = ValidateInput(InputShapes[0], InputShapes[1], InputShapes[2]); Res < 0)
+			if (const int Res = ValidateInput(InputTensors[0]->GetShape(), InputTensors[1]->GetShape(), InputTensors[2]->GetShape()); Res < 0)
 			{
 				return Res;
 			}
 
-			OutputShapes.Add(InputShapes[0]);
+			OutputTensors[0]->SetShape(InputTensors[0]->GetShape());
 
 			return 0;
 		};
diff --git a/Engine/Plugins/Experimental/NNE/Source/NNXRuntimeRDG/Private/Hlsl/NNIRuntimeRDGMatMul.cpp b/Engine/Plugins/Experimental/NNE/Source/NNXRuntimeRDG/Private/Hlsl/NNIRuntimeRDGMatMul.cpp
--- a/Engine/Plugins/Experimental/NNE/Source/NNXRuntimeRDG/Private/Hlsl/NNIRuntimeRDGMatMul.cpp
+++ b/Engine/Plugins/Experimental/NNE/Source/NNXRuntimeRDG/Private/Hlsl/NNIRuntimeRDGMatMul.cpp
@@ -20,13 +20,13 @@ namespace UE::NNIRuntimeRDG::Private::Hlsl
 
 	public:
 
-		virtual int ComputeOutputShape(TConstArrayView<NNX::FTensorShape> InputShapes, TArray<NNX::FTensorShape>& OutputShapes) const override
+		virtual int PrepareOutputs(TConstArrayView<NNX::FTensorRef> InputTensors, TArrayView<NNX::FTensorRef> OutputTensors) const override
 		{
-			OutputShapes.Empty();
-			check(InputShapes.Num() == 2);
+			check(InputTensors.Num() == 2);
+			check(OutputTensors.Num() == 1);
 
-			const NNX::FTensorShape& InputA = InputShapes[0];
-			const NNX::FTensorShape& InputB = InputShapes[1];
+			const NNX::FTensorShape& InputA = InputTensors[0]->GetShape();
+			const NNX::FTensorShape& InputB = InputTensors[1]->GetShape();
 
 			if (InputA.Rank() < 2)
 			{
@@ -64,7 +64,7 @@ namespace UE::NNIRuntimeRDG::Private::Hlsl
 			OutputShape.Data[OutputRank - 2] = InputA.Data[InputA.Rank() - 2];
 			OutputShape.Data[OutputRank - 1] = InputB.Data[InputB.Rank() - 1];
 
-			OutputShapes.Emplace(OutputShape);
+			OutputTensors[0]->SetShape(OutputShape);
 
 			return 0;
 		};
diff --git a/Engine/Plugins/Experimental/NNE/Source/NNXRuntimeRDG/Private/NNXRuntimeDML.cpp b/Engine/Plugins/Experimental/NNE/Source/NNXRuntimeRDG/Private/NNXRuntimeDML.cpp
--- a/Engine/Plugins/Experimental/NNE/Source/NNXRuntimeRDG/Private/NNXRuntimeDML.cpp
+++ b/Engine/Plugins/Experimental/NNE/Source/NNXRuntimeRDG/Private/NNXRuntimeDML.cpp
@@ -329,7 +329,7 @@ public:
 //
 // DirectML operator
 //
-class FMLOperatorDml : public FMLOperatorRDG
+class FMLOperatorDml : public IOperatorRDG
 {
 public:
 
@@ -1186,7 +1186,7 @@ public:
 protected:
 
 	virtual void AddDispatchOps_RenderThread(FRDGBuilder& GraphBuilder) override;
-	virtual int RunShapeInference() override;
+	virtual int PrepareTensorShapesAndData() override;
 
 private:
 
@@ -1551,8 +1551,6 @@ bool FMLInferenceModelDml::Init(TConstArrayView<uint8> ModelData, FDeviceContext
 //
 void FMLInferenceModelDml::AddDispatchOps_RenderThread(FRDGBuilder& GraphBuilder)
 {
-	check(AllTensorRDGs.Num() == AllShapes.Num());
-
 	static constexpr int32 MaxExpectedInput = 10;
 	TArray<FTensorRDGRef, TInlineAllocator<MaxExpectedInput>> InputTensors;
 
@@ -1610,24 +1608,16 @@ FMLOperatorDml* FMLInferenceModelDml::OpCreate(const FString& OpName, TArrayView
 //
 //
 //
-int FMLInferenceModelDml::RunShapeInference()
+int FMLInferenceModelDml::PrepareTensorShapesAndData()
 {
-	AllShapes.Empty();
-
 	for (FTensorDesc SymbolicTensorDesc : AllSymbolicTensorDescs)
 	{
-		if (SymbolicTensorDesc.IsConcrete())
+		if (!SymbolicTensorDesc.IsConcrete())
 		{
-			FTensorShape TensorShape = FTensorShape::MakeFromSymbolic(SymbolicTensorDesc.GetShape());
-			AllShapes.Emplace(TensorShape);
+			UE_LOG(LogNNX, Warning, TEXT("DML engine does not support model with variable shapes yet."));
+			return -1;
 		}
 	}
-	if (AllShapes.Num() != AllSymbolicTensorDescs.Num())
-	{
-		AllShapes.Empty();
-		UE_LOG(LogNNX, Warning, TEXT("DML engine does not support model with variable shapes yet."));
-		return -1;
-	}
 
 	return 0;
 }
diff --git a/Engine/Plugins/Experimental/NNE/Source/NNXRuntimeRDG/Private/NNXRuntimeHLSL.cpp b/Engine/Plugins/Experimental/NNE/Source/NNXRuntimeRDG/Private/NNXRuntimeHLSL.cpp
--- a/Engine/Plugins/Experimental/NNE/Source/NNXRuntimeRDG/Private/NNXRuntimeHLSL.cpp
+++ b/Engine/Plugins/Experimental/NNE/Source/NNXRuntimeRDG/Private/NNXRuntimeHLSL.cpp
@@ -36,7 +36,7 @@ public:
 protected:
 
 	virtual void AddDispatchOps_RenderThread(FRDGBuilder& GraphBuilder) override;
-	virtual int RunShapeInference() override;
+	virtual int PrepareTensorShapesAndData() override;
 
 private:
 
@@ -137,104 +137,75 @@ FMLOperatorHlsl* FMLInferenceModelHlsl::OpCreate(const FString& OpName, TConstAr
 	return Op;
 }
 
-int FMLInferenceModelHlsl::RunShapeInference()
+int FMLInferenceModelHlsl::PrepareTensorShapesAndData()
 {
-	AllShapes.Empty();
-
+	check(AllTensorRDGs.Num() == AllSymbolicTensorDescs.Num());
+	
 	if (Operators.Num() == 0)
 	{
 		UE_LOG(LogNNX, Warning, TEXT("No operators in model"));
 		return -1;
 	}
 
+	// Run model preparation (including shape inference) on all operators
+	// This loop could be abstracted to a different engine/system as it apply on FTensorRef & IPrepareOperator witch are RDG agnostics.
 	static constexpr int32 MaxExpectedInput = 10;
-	TArray<FTensorShape, TInlineAllocator<MaxExpectedInput>> InputShapes;
-	TArray<FTensorShape> OutputShapes;
-	TArray<bool> AllInitializedShapes;
-
-	checkCode(AllInitializedShapes.SetNum(AllSymbolicTensorDescs.Num(), false););
-	AllShapes.SetNum(AllSymbolicTensorDescs.Num());
-
-	//Prime shape inference with model inputs
-	for (int32 i = 0; i < InputTensorIndices.Num(); ++i)
-	{
-		const int32 InputTensorModelIndex = InputTensorIndices[i];
+	TArray<FTensorRef, TInlineAllocator<MaxExpectedInput>> InputTensors;
+	TArray<FTensorRef> OutputTensors;
+	TArray<bool> AllInitializedTensors;
 
-		AllShapes[InputTensorModelIndex] = InputTensorShapes[i];
-		checkCode(AllInitializedShapes[InputTensorModelIndex] = true);
-	}
-
-	//Prime shape inference with model weights
-	for (int32 i = 0; i < WeightTensorIndices.Num(); ++i)
-	{
-		const int32 WeightTensorModelIndex = WeightTensorIndices[i];
-
-		AllShapes[WeightTensorModelIndex] = WeightTensorRDGs[i].GetShape();
-		checkCode(AllInitializedShapes[WeightTensorModelIndex] = true);
-	}
+	checkCode(
+		AllInitializedTensors.Init(false, AllSymbolicTensorDescs.Num());
+		for (int32 Idx : InputTensorIndices)
+		{
+			AllInitializedTensors[Idx] = true;
+		}
+		for (int32 Idx : WeightTensorIndices)
+		{
+			AllInitializedTensors[Idx] = true;
+		}
+	);
 
-	// Run shape inference on all operators
+	// Run model preparation (including shape inference) on all operators
+	// This loop could be abstracted to a different system as it apply on FTensorRef & IPrepareOperator witch are RDG agnostics.
 	for (int32 Idx = 0; Idx < Operators.Num(); ++Idx)
 	{
-		InputShapes.Empty();
-		OutputShapes.Empty();
+		InputTensors.Empty();
+		OutputTensors.Empty();
 
+		//Operator inputs
 		for (int32 i : OperatorInputTensorIndices[Idx])
 		{
-			checkf(AllInitializedShapes[i] == true, TEXT("Input shape %d for operator %d should have been initialized."), i, Idx);
-			InputShapes.Emplace(AllShapes[i]);
+			checkf(AllInitializedTensors[i] == true, TEXT("Input tensor %d for operator %d should have been initialized."), i, Idx);
+			InputTensors.Emplace(AllTensorRDGs[i]);
 		}
-
-		//If operator output shapes are not variable we don't need to run shape inference
-		bool bShouldRunShapeInferenceForOperator = false;
-		for (int32 i = 0; i < OperatorOutputTensorIndices[Idx].Num(); ++i)
+		//Operator outputs
+		for (int32 i : OperatorOutputTensorIndices[Idx])
 		{
-			const int32 OutputTensorIndex = OperatorOutputTensorIndices[Idx][i];
-
-			if (!AllSymbolicTensorDescs[OutputTensorIndex].IsConcrete())
-			{
-				bShouldRunShapeInferenceForOperator = true;
-				continue;
-			}
-			OutputShapes.Emplace(FTensorShape::MakeFromSymbolic(AllSymbolicTensorDescs[OutputTensorIndex].GetShape()));
+			OutputTensors.Emplace(AllTensorRDGs[i]);
+			checkf(AllInitializedTensors[i] == false, TEXT("Output tensor %d for operator %d should not have been initialized yet."), i, Idx);
+			checkCode(AllInitializedTensors[i] = true);
 		}
 
-		//Otherwise we need to run shape inference for the operator
-		if (bShouldRunShapeInferenceForOperator)
-		{
-			OutputShapes.Empty();
-
-			const FMLOperatorHlsl* Op = Operators[Idx];
-
-			if (Op->ComputeOutputShape(InputShapes, OutputShapes) != 0)
-			{
-				//Operator could not compute output shapes, meaning we can't allocate
-				//output buffer before running the model. This engine does not support this.
-				UE_LOG(LogNNX, Warning, TEXT("Could not deduce tensor shapes for this model during shape inference, HLSL engine wont support the model as it need to precompute all shapes for performance reasons."));
-				AllShapes.Empty();
-				OutputTensorShapes.Empty();
-				return -1;
-			}
-		}
+		const FMLOperatorHlsl* Op = Operators[Idx];
 
-		for (int32 i = 0; i < OperatorOutputTensorIndices[Idx].Num(); ++i)
+		if (Op->PrepareOutputs(InputTensors, OutputTensors) != 0)
 		{
-			const int32 OutputTensorIndex = OperatorOutputTensorIndices[Idx][i];
-
-			AllShapes[OutputTensorIndex] = OutputShapes[i];
-			checkCode(AllInitializedShapes[OutputTensorIndex] = true);
+			//Operator could not prepare the output tensors, meaning we can't allocate
+			//output buffer before running the model. This engine does not support this.
+			UE_LOG(LogNNX, Warning, TEXT("Could not deduce tensor shapes for this model during shape inference, HLSL engine wont support the model as it need to precompute all shapes for performance reasons."));
+			AllTensorRDGs.Empty();
+			return -1;
 		}
 	}
 
 	checkCode(
-		for (int i = 0; i < AllInitializedShapes.Num(); ++i)
+		for (int i = 0; i < AllInitializedTensors.Num(); ++i)
 		{
-			checkf(AllInitializedShapes[i], TEXT("Tensor at index %d, was not initialized by shape inference."));
+			checkf(AllInitializedTensors[i], TEXT("Tensor at index %d, was not initialized by model preparation."));
 		};
 	);
 
-	check(AllShapes.Num() == AllSymbolicTensorDescs.Num());
-
 	return 0;
 }
 
@@ -243,8 +214,6 @@ int FMLInferenceModelHlsl::RunShapeInference()
 //
 void FMLInferenceModelHlsl::AddDispatchOps_RenderThread(FRDGBuilder& GraphBuilder)
 {
-	check(AllTensorRDGs.Num() == AllShapes.Num());
-
 	static constexpr int32 MaxExpectedInput = 10;
 	TArray<FTensorRDGRef, TInlineAllocator<MaxExpectedInput>> InputTensors;
 
diff --git a/Engine/Plugins/Experimental/NNE/Source/NNXRuntimeRDG/Private/NNXRuntimeHLSLOp.h b/Engine/Plugins/Experimental/NNE/Source/NNXRuntimeRDG/Private/NNXRuntimeHLSLOp.h
--- a/Engine/Plugins/Experimental/NNE/Source/NNXRuntimeRDG/Private/NNXRuntimeHLSLOp.h
+++ b/Engine/Plugins/Experimental/NNE/Source/NNXRuntimeRDG/Private/NNXRuntimeHLSLOp.h
@@ -10,12 +10,11 @@ namespace NNX
 	/**
 	 * Base class for all Hlsl ML operators
 	 */
-	struct FMLOperatorHlsl : public FMLOperatorRDG
+	struct FMLOperatorHlsl : public IOperatorRDG, public IPrepareOperator
 	{
 		virtual ~FMLOperatorHlsl() = default;
 		virtual bool Initialize(TConstArrayView<FTensorDesc> InputTensorDescs, TConstArrayView<FTensorDesc> OutputTensorDescs, const UE::NNECore::FAttributeMap& Attributes) = 0;
 		virtual void Dispatch(FRDGBuilder& GraphBuilder, TConstArrayView<FTensorRDGRef> InInputTensors, TConstArrayView<FTensorRDGRef> InOutputTensors) = 0;
-		virtual int ComputeOutputShape(TConstArrayView<FTensorShape> InputShapes, TArray<FTensorShape>& OutputShapes) const { return -1; };
 	};
 
 	/**
diff --git a/Engine/Plugins/Experimental/NNE/Source/NNXRuntimeRDG/Private/NNXRuntimeRDG.cpp b/Engine/Plugins/Experimental/NNE/Source/NNXRuntimeRDG/Private/NNXRuntimeRDG.cpp
--- a/Engine/Plugins/Experimental/NNE/Source/NNXRuntimeRDG/Private/NNXRuntimeRDG.cpp
+++ b/Engine/Plugins/Experimental/NNE/Source/NNXRuntimeRDG/Private/NNXRuntimeRDG.cpp
@@ -363,23 +363,81 @@ int FMLInferenceModelRDG::SetInputTensorShapes(TConstArrayView<FTensorShape> InI
 		return -1;
 	}
 
-	//Run shape inference filling AllTensors
-	if (RunShapeInference() != 0)
+	//Allocate and prime all AllTensorRDGs with concrete shapes defaulting variables dimension to 1 if needed
+	AllTensorRDGs.Init(nullptr, AllSymbolicTensorDescs.Num());
+
+	InputTensorRDGs.Empty();
+	for (int32 i = 0; i < InputTensorIndices.Num(); ++i)
+	{
+		const int32 Idx = InputTensorIndices[i];
+		const FTensorDesc& TensorDesc = InputSymbolicTensors[i];
+		const FTensorShape& TensorShape = InputTensorShapes[i];
+		
+		InputTensorRDGs.Emplace(FTensorRDG::Make(TensorDesc, TensorShape, nullptr));
+		AllTensorRDGs[Idx] = &InputTensorRDGs[i];
+	}
+
+	for (int32 i = 0; i < WeightTensorIndices.Num(); ++i)
+	{
+		const int32 Idx = WeightTensorIndices[i];
+		
+		AllTensorRDGs[Idx] = &WeightTensorRDGs[i];
+	}
+
+	IntermediateTensorRDGs.Empty();
+	for (int32 i = 0; i < IntermediateTensorIndices.Num(); ++i)
+	{
+		const int32 Idx = IntermediateTensorIndices[i];
+		const FTensorDesc& TensorDesc = AllSymbolicTensorDescs[Idx];
+		const FTensorShape TensorShape = FTensorShape::MakeFromSymbolic(TensorDesc.GetShape());
+		
+		IntermediateTensorRDGs.Emplace(FTensorRDG::Make(TensorDesc, TensorShape, nullptr));
+		AllTensorRDGs[Idx] = &IntermediateTensorRDGs[i];
+	}
+
+	OutputTensorRDGs.Empty();
+	for (int32 i = 0; i < OutputTensorIndices.Num(); ++i)
+	{
+		const int32 Idx = OutputTensorIndices[i];
+		const FTensorDesc& TensorDesc = OutputSymbolicTensors[i];
+		const FTensorShape TensorShape = FTensorShape::MakeFromSymbolic(TensorDesc.GetShape());
+		
+		OutputTensorRDGs.Emplace(FTensorRDG::Make(TensorDesc, TensorShape, nullptr));
+		AllTensorRDGs[Idx] = &OutputTensorRDGs[i];
+	}
+
+	checkCode(
+		for (int i = 0; i < AllTensorRDGs.Num(); ++i)
+		{
+			checkf(AllTensorRDGs[i] != nullptr, TEXT("Tensor at index %d, was not allocated for model preparation."), i);
+		};
+	);
+
+	//Allow the specific engine to run shape inference if supported
+	if (PrepareTensorShapesAndData() != 0)
 	{
 		return -1;
 	}
 
-	//Set OutputTensorShapes for the model from shape inference result
+	checkCode(
+		for (int i = 0; i < AllTensorRDGs.Num(); ++i)
+		{
+			checkf(AllTensorRDGs[i] != nullptr, TEXT("Tensor at index %d, was not allocated after model preparation."), i);
+			checkf(AllTensorRDGs[i]->GetShape().IsCompatibleWith(AllSymbolicTensorDescs[i].GetShape()), TEXT("Tensor at index %d have a shape incompatible with model definition."), i);
+		};
+	);
+
+	//Set OutputTensorShapes for the model from preparation result
 	for (int32 OutputIndices : OutputTensorIndices)
 	{
-		OutputTensorShapes.Emplace(AllShapes[OutputIndices]);
+		OutputTensorShapes.Emplace(AllTensorRDGs[OutputIndices]->GetShape());
 	}
 
-	check(InputTensorIndices.Num() + OutputTensorIndices.Num() + WeightTensorIndices.Num() + IntermediateTensorIndices.Num() == AllShapes.Num());
+	check(InputTensorIndices.Num() + OutputTensorIndices.Num() + WeightTensorIndices.Num() + IntermediateTensorIndices.Num() == AllTensorRDGs.Num());
 	check(InputTensorShapes.Num() == InputSymbolicTensors.Num());
 	check(OutputTensorShapes.Num() == OutputSymbolicTensors.Num());
 	check(WeightTensorIndices.Num() == WeightTensorRDGs.Num());
-	check(AllShapes.Num() == AllSymbolicTensorDescs.Num());
+	check(AllTensorRDGs.Num() == AllSymbolicTensorDescs.Num());
 	
 	return 0;
 }
@@ -410,10 +468,8 @@ int FMLInferenceModelRDG::EnqueueRDG(FRDGBuilder& RDGBuilder, TConstArrayView<FM
 	}
 
 	// Process input binding, and if required, allocate RDG buffers
-	InputTensorRDGs.Empty();
 	FIntArray		InputUploadIndices;
-
-	Res = SetTensors(RDGBuilder, InputTensorRDGs, InputUploadIndices, InInputBindings, InputSymbolicTensors, InputTensorShapes);
+	Res = SetTensors(RDGBuilder, InputTensorRDGs, InputUploadIndices, InInputBindings);
 	if (Res != 0)
 	{
 		UE_LOG(LogNNX, Warning, TEXT("Invalid input tensor binding type for tensor index:%d"), Res);
@@ -421,10 +477,9 @@ int FMLInferenceModelRDG::EnqueueRDG(FRDGBuilder& RDGBuilder, TConstArrayView<FM
 	}
 
 	// Process output tensors, and if required, allocate RDG buffers
-	OutputTensorRDGs.Empty();
 	FIntArray		OutputReadbackIndices;
 
-	Res = SetTensors(RDGBuilder, OutputTensorRDGs, OutputReadbackIndices, InOutputBindings, OutputSymbolicTensors, OutputTensorShapes);
+	Res = SetTensors(RDGBuilder, OutputTensorRDGs, OutputReadbackIndices, InOutputBindings);
 	if (Res != 0)
 	{
 		UE_LOG(LogNNX, Warning, TEXT("Invalid output tensor binding type for tensor index:%d"), Res);
@@ -434,34 +489,13 @@ int FMLInferenceModelRDG::EnqueueRDG(FRDGBuilder& RDGBuilder, TConstArrayView<FM
 	// If required, upload input tensors to GPU
 	AddTensorUploads_RenderThread(RDGBuilder, InputUploadIndices, InputTensorRDGs, InInputBindings);
 
-	AllTensorRDGs.Empty();
-	AllTensorRDGs.SetNum(AllShapes.Num());
-	
-	//Create intermediate tensors
-	IntermediateTensorRDGs.Empty();
-	for (int32 i = 0; i < IntermediateTensorIndices.Num(); ++i)
+	//Create buffer for intermediate tensors
+	for (FTensorRDG& TensorRDG : IntermediateTensorRDGs)
 	{
-		const int32 Idx = IntermediateTensorIndices[i];
-		const FTensorDesc& TensorDesc = AllSymbolicTensorDescs[Idx];
-		const FTensorShape& TensorShape = AllShapes[Idx];
-		FTensorRDG& TensorRDG = IntermediateTensorRDGs.Emplace_GetRef(FTensorRDG::Make(TensorDesc, TensorShape, nullptr));
 		const FRDGBufferDesc BufferDesc = CreateRDGBufferDescForTensorRDG(TensorRDG);
 		const FRDGBufferRef TensorBuffer = RDGBuilder.CreateBuffer(BufferDesc, *TensorRDG.GetName(), ERDGBufferFlags::None);
-		
+		check(TensorRDG.GetBuffer() == nullptr);
 		TensorRDG.SetBuffer(TensorBuffer);
-		AllTensorRDGs[Idx] = &IntermediateTensorRDGs[i];
-	}
-
-	//Insert input tensors
-	for (int32 i = 0; i < InputTensorIndices.Num(); ++i)
-	{
-		AllTensorRDGs[InputTensorIndices[i]] = &InputTensorRDGs[i];
-	}
-
-	//Insert output tensors
-	for (int32 i = 0; i < OutputTensorIndices.Num(); ++i)
-	{
-		AllTensorRDGs[OutputTensorIndices[i]] = &OutputTensorRDGs[i];
 	}
 
 	//For now weights tensors are not uploaded to GPU thus GetBuffer will return nullptr for them.
@@ -486,36 +520,31 @@ int FMLInferenceModelRDG::EnqueueRDG(FRDGBuilder& RDGBuilder, TConstArrayView<FM
  * Process binding and check if we need to create RDG Buffer for CPU binding 
  * Returns 0 on success, or index of a tensor if the tensor type is not supported.
  */
-int FMLInferenceModelRDG::SetTensors(FRDGBuilder& GraphBuilder, FTensorRDGArray& OutTensorRDGs, FIntArray& OutIndices, 
-	TConstArrayView<FMLTensorBinding> InBindings, TConstArrayView<FTensorDesc> InTensorDescs, TConstArrayView<FTensorShape> InTensorShapes)
+int FMLInferenceModelRDG::SetTensors(FRDGBuilder& GraphBuilder, FTensorRDGArray& InTensorRDGs, FIntArray& OutIndices, 
+	TConstArrayView<FMLTensorBinding> InBindings)
 {
-	check(InBindings.Num() == InTensorDescs.Num());
-	check(InBindings.Num() == InTensorShapes.Num());
+	check(InBindings.Num() == InTensorRDGs.Num());
 	
 	for (int32 Idx = 0; Idx < InBindings.Num(); ++Idx)
 	{
+		FTensorRDG& TensorRDG = InTensorRDGs[Idx];
 		const FMLTensorBinding& Binding = InBindings[Idx];
-		const FTensorDesc& TensorDesc = InTensorDescs[Idx];
-		const FTensorShape& TensorShape = InTensorShapes[Idx];
 
 		if (Binding.BindingType == EMLTensorBindingDataType::CPUMemory)
 		{
-			FTensorRDG TensorRDG = FTensorRDG::Make(TensorDesc, TensorShape, nullptr);
 			FRDGBufferDesc Desc = CreateRDGBufferDescForTensorRDG(TensorRDG);
 			
 			// FIXME: We should use BUF_SourceCopy for only output buffers (GPU readback)
 			Desc.Usage = EBufferUsageFlags(Desc.Usage | BUF_SourceCopy);
 			
-			FRDGBufferRef TensorBuffer = GraphBuilder.CreateBuffer(Desc, *TensorDesc.GetName(), ERDGBufferFlags::None);
+			FRDGBufferRef TensorBuffer = GraphBuilder.CreateBuffer(Desc, *TensorRDG.GetName(), ERDGBufferFlags::None);
 			
 			TensorRDG.SetBuffer(TensorBuffer);
-			OutTensorRDGs.Add(TensorRDG);
 			OutIndices.Add(Idx);
 		}
 		else if (Binding.BindingType == EMLTensorBindingDataType::RDGBuffer)
 		{
-			FTensorRDG TensorRDG = FTensorRDG::Make(TensorDesc, TensorShape, Binding.Buffer);
-			OutTensorRDGs.Add(TensorRDG);
+			TensorRDG.SetBuffer(Binding.Buffer);
 		}
 		else
 		{
diff --git a/Engine/Plugins/Experimental/NNE/Source/NNXRuntimeRDG/Private/NNXRuntimeRDG.h b/Engine/Plugins/Experimental/NNE/Source/NNXRuntimeRDG/Private/NNXRuntimeRDG.h
--- a/Engine/Plugins/Experimental/NNE/Source/NNXRuntimeRDG/Private/NNXRuntimeRDG.h
+++ b/Engine/Plugins/Experimental/NNE/Source/NNXRuntimeRDG/Private/NNXRuntimeRDG.h
@@ -37,10 +37,21 @@ namespace NNX
 struct FTensorBinding;
 class FTensorDesc;
 
-/** Base class for all ML operators running on the RDG */
-struct FMLOperatorRDG
+/**
+ * Interface for all operators to prepare the model tensors at scheduling time
+ */
+struct IPrepareOperator
+{
+	virtual ~IPrepareOperator() = default;
+	virtual int PrepareOutputs(TConstArrayView<NNX::FTensorRef> InputTensors, TArrayView<NNX::FTensorRef> OutputTensors) const = 0;
+};
+
+/**
+* Interface for all ML operators running on the RDG 
+*/
+struct IOperatorRDG
 {
-	virtual ~FMLOperatorRDG() = default;
+	virtual ~IOperatorRDG() = default;
 };
 
 class FTensorRDG : public FTensor
@@ -57,6 +68,7 @@ public:
 		TensorRDG.Shape = Shape;
 		TensorRDG.Volume = Shape.Volume();
 		TensorRDG.DataSize = (uint64)GetTensorDataTypeSizeInBytes(TensorRDG.DataType) * TensorRDG.Volume;
+		TensorRDG.Buffer = Buffer;
 		check(TensorRDG.Volume <= TNumericLimits<uint32>::Max());
 		return TensorRDG;
 	}
@@ -103,9 +115,9 @@ protected:
 
 	bool LoadModel(TConstArrayView<uint8> ModelData, FMLRuntimeFormat& Format);
 
-	int SetTensors(FRDGBuilder& GraphBuilder, FTensorRDGArray& OutTensorRDGs, FIntArray& OutIndices, TConstArrayView<FMLTensorBinding> InBindings, TConstArrayView<FTensorDesc> InTensorDescs, TConstArrayView<FTensorShape> InTensorShapes);
+	int SetTensors(FRDGBuilder& GraphBuilder, FTensorRDGArray& InTensorRDGs, FIntArray& OutIndices, TConstArrayView<FMLTensorBinding> InBindings);
 	
-	virtual int RunShapeInference() = 0;
+	virtual int PrepareTensorShapesAndData() = 0;
 	virtual void AddDispatchOps_RenderThread(FRDGBuilder& GraphBuilder) = 0;
 
 	virtual void AddTensorUploads_RenderThread(FRDGBuilder& GraphBuilder, TConstArrayView<int32> InUploadIndices, TConstArrayView<FTensorRDG> InRDGBindings, TConstArrayView<FMLTensorBinding> InBindings);
@@ -113,7 +125,6 @@ protected:
 
 	//Tensor descriptor
 	TArray<FTensorDesc>			AllSymbolicTensorDescs;
-	TArray<FTensorShape>		AllShapes;
 
 	//Tensor indices for models
 	TArray<int32>				IntermediateTensorIndices;
@@ -171,7 +182,7 @@ private:
 		{
 		}
 
-		//TODO should be extended as needed by operator to support more validation especially around the range of the values.
+		//Idea: we could extended as needed by operator to support more validation especially around the range of the values.
 		//An example is ConvTranspose `auto_pad` enum style string that can only take a few values
 		//In the same direction we might only support a range of value for a float (for example
 		//we only support integer but the type is float, or only positive values for an int32)
@@ -293,12 +304,14 @@ public:
 
 			const FString& OpType = Format.Operators[Idx].TypeName;
 			
-			//TODO jira 167587: we should extract constant tensor from the model
-			//and pass them to the operator validation so that it can validate
-			//the shapes of the constant tensors and ensure no gpu-cpu sync
-			//will be needed during the execution of the operator.
 			typename TOperatorRegistryRDG<TOperatorType>::OperatorValidateFunc ValidationFn = Registry->OpFindValidation(OpType);
 
+			if (!ValidationFn)
+			{
+				UE_LOG(LogNNX, Warning, TEXT("Hlsl MLOperatorRegistry failed to find validation for operator:%s"), *OpType);
+				return false;
+			}
+			
 			if (!ValidationFn(AttributeMap, InputTensorTypes, InputTensorShapes))
 			{
 				UE_LOG(LogNNX, Warning, TEXT("Hlsl MLOperatorRegistry failed to validate operator:%s"), *OpType);
