NNE Support for Upsample operator in HLSL (nearest interpolation mode only atm)

Problem:
* Upsample needed for style transfer was not implemented

Solution:
* Added a first version of upsample only supporting nearest interpolation for now

Tests:
* All NNX tests pass under win64 editor + added upsample tests for integer upscaling with various tensor ranks
* ONNX CPU upsample is giving weird (i would call that a bug) result for not integer upscaling. Thus i didn't add tests for those cases as we use ORT CPU as a reference. (ORTDml result also confirm ORTCpu result are off in that case).

diff --git a/Engine/Plugins/Experimental/NNE/Source/NNEHlslShaders/Internal/NNEHlslShadersUpsampleCS.h b/Engine/Plugins/Experimental/NNE/Source/NNEHlslShaders/Internal/NNEHlslShadersUpsampleCS.h
new file mode 100644
--- /dev/null
+++ b/Engine/Plugins/Experimental/NNE/Source/NNEHlslShaders/Internal/NNEHlslShadersUpsampleCS.h
@@ -0,0 +1,40 @@
+// Copyright Epic Games, Inc. All Rights Reserved.
+
+#pragma once
+
+#include "GlobalShader.h"
+#include "ShaderParameterUtils.h"
+#include "RenderGraphUtils.h"
+#include "NNXOperator.h"
+
+namespace UE::NNEHlslShaders::Internal
+{
+	class FUpsampleConstants
+	{
+	public:
+
+		static const int32 MAX_NUM_DIMENSIONS{ 8 };
+		static const int32 NUM_GROUP_THREADS{ 256 };
+	};
+
+	class NNEHLSLSHADERS_API FUpsampleCS : public FGlobalShader
+	{
+		DECLARE_GLOBAL_SHADER(FUpsampleCS);
+		SHADER_USE_PARAMETER_STRUCT(FUpsampleCS, FGlobalShader)
+
+		class FUpsampleNumDimensions : SHADER_PERMUTATION_RANGE_INT("NUM_DIMENSIONS", 1, FUpsampleConstants::MAX_NUM_DIMENSIONS);
+		using FPermutationDomain = TShaderPermutationDomain<FUpsampleNumDimensions>;
+
+	public:
+
+		BEGIN_SHADER_PARAMETER_STRUCT(FParameters, )
+			SHADER_PARAMETER_RDG_BUFFER_SRV(Buffer<float>, Input)
+			SHADER_PARAMETER_RDG_BUFFER_UAV(RWBuffer<float>, Output)
+			SHADER_PARAMETER_ARRAY(FUintVector4, TensorInfo, [FUpsampleConstants::MAX_NUM_DIMENSIONS])
+			SHADER_PARAMETER(uint32, Num)
+			SHADER_PARAMETER(uint32, ThreadCountX)
+		END_SHADER_PARAMETER_STRUCT()
+
+			static void ModifyCompilationEnvironment(const FGlobalShaderPermutationParameters& InParameters, FShaderCompilerEnvironment& OutEnvironment);
+	};
+} // UE::NNEHlslShaders::Internal
diff --git a/Engine/Plugins/Experimental/NNE/Source/NNEHlslShaders/Private/NNEHlslShadersUpsampleCS.cpp b/Engine/Plugins/Experimental/NNE/Source/NNEHlslShaders/Private/NNEHlslShadersUpsampleCS.cpp
new file mode 100644
--- /dev/null
+++ b/Engine/Plugins/Experimental/NNE/Source/NNEHlslShaders/Private/NNEHlslShadersUpsampleCS.cpp
@@ -0,0 +1,18 @@
+// Copyright Epic Games, Inc. All Rights Reserved.
+
+#include "NNEHlslShadersUpsampleCS.h"
+#include "NNXCore.h"
+
+namespace UE::NNEHlslShaders::Internal
+{
+	void FUpsampleCS::ModifyCompilationEnvironment(const FGlobalShaderPermutationParameters& InParameters, FShaderCompilerEnvironment& OutEnvironment)
+	{
+		FGlobalShader::ModifyCompilationEnvironment(InParameters, OutEnvironment);
+
+		OutEnvironment.SetDefine(TEXT("THREADGROUP_SIZE_X"), FUpsampleConstants::NUM_GROUP_THREADS);
+
+		FPermutationDomain PermutationVector(InParameters.PermutationId);
+	}
+
+	IMPLEMENT_GLOBAL_SHADER(FUpsampleCS, "/NNE/NNEHlslShadersUpsample.usf", "Upsample", SF_Compute);
+} // UE::NNEHlslShaders::Internal
\ No newline at end of file
diff --git a/Engine/Plugins/Experimental/NNE/Source/NNEHlslShaders/Shaders/NNEHlslShadersUpsample.usf b/Engine/Plugins/Experimental/NNE/Source/NNEHlslShaders/Shaders/NNEHlslShadersUpsample.usf
new file mode 100644
--- /dev/null
+++ b/Engine/Plugins/Experimental/NNE/Source/NNEHlslShaders/Shaders/NNEHlslShadersUpsample.usf
@@ -0,0 +1,48 @@
+// Copyright Epic Games, Inc. All Rights Reserved.
+
+#include "/Engine/Public/Platform.ush"
+#include "NNEHlslShadersBroadcastHelper.ush"
+
+Buffer<float> Input;
+RWBuffer<float> Output;
+uint4 TensorInfo[NUM_DIMENSIONS];
+uint Num;
+uint ThreadCountX;
+
+#define INPUT_STRIDE 0
+#define OUTPUT_STRIDE 1
+#define INPUT_SIZE 2
+#define OUTPUT_SIZE 3
+
+uint GetNearestIndexFromUV(float UV, uint ArraySize)
+{
+	uint Index = round(UV * ArraySize - 0.5f);
+	return min(Index, ArraySize - 1);
+}
+
+[numthreads(THREADGROUP_SIZE_X, 1, 1)]
+void Upsample(in const uint3 DispatchThreadID : SV_DispatchThreadID)
+{
+	const uint Index = DispatchThreadID.y * ThreadCountX + DispatchThreadID.x;
+	
+	if (Index < Num)
+	{
+		uint InputIndex = 0;
+
+		uint Offset = Index;
+		for (uint dim = 0; dim < NUM_DIMENSIONS; ++dim)
+		{
+			uint OutDimIdx, R;
+			DivMod(Offset, TensorInfo[dim][OUTPUT_STRIDE], OutDimIdx, R);
+			
+			float UV = (0.5f + OutDimIdx) / TensorInfo[dim][OUTPUT_SIZE];
+			uint InputDimIdx = GetNearestIndexFromUV(UV, TensorInfo[dim][INPUT_SIZE]);
+			
+			InputIndex += TensorInfo[dim][INPUT_STRIDE] * InputDimIdx;
+			Offset = R;
+		}
+
+		float X = Input[InputIndex];
+		Output[Index] = X;
+	}
+}
diff --git a/Engine/Plugins/Experimental/NNE/Source/NNXCore/Public/NNXTypes.h b/Engine/Plugins/Experimental/NNE/Source/NNXCore/Public/NNXTypes.h
--- a/Engine/Plugins/Experimental/NNE/Source/NNXCore/Public/NNXTypes.h
+++ b/Engine/Plugins/Experimental/NNE/Source/NNXCore/Public/NNXTypes.h
@@ -247,9 +247,13 @@ public:
 		return Shape;
 	}
 
-	TConstArrayView<uint8> GetPreparedData() const
+	template <typename T> TConstArrayView<T> GetPreparedData() const
 	{
-		return PreparedData;
+		const T* DataPtr = reinterpret_cast<const T*>(PreparedData.GetData());
+		const int32 ElemSize = sizeof(T);
+		
+		check(DataSize % ElemSize == 0);
+		return MakeArrayView(DataPtr, DataSize / ElemSize);
 	}
 
 	void SetShape(const FTensorShape& InShape)
@@ -268,7 +272,7 @@ public:
 		PreparedData.Append(Data);
 	}
 
-	bool HasPreparedData()
+	bool HasPreparedData() const
 	{
 		return !PreparedData.IsEmpty();
 	}
diff --git a/Engine/Plugins/Experimental/NNE/Source/NNXQA/Resources/NNXQATestDesc.json b/Engine/Plugins/Experimental/NNE/Source/NNXQA/Resources/NNXQATestDesc.json
--- a/Engine/Plugins/Experimental/NNE/Source/NNXQA/Resources/NNXQATestDesc.json
+++ b/Engine/Plugins/Experimental/NNE/Source/NNXQA/Resources/NNXQATestDesc.json
@@ -311,6 +311,15 @@
 			"targets": [
 				{ "target": "InstanceNormalization" }
 			]
+		},
+		{
+			"category": "Upsample",
+			"runtimes": [
+				{ "name": "NNXRuntimeDml", "skip": true	}
+			],
+			"targets": [
+				{ "target": "Upsample" }
+			]
 		}
 	],
 
@@ -565,7 +574,7 @@
 		{
 			"name": "MatMul.Rank5",
 			"runtimes": [	
-				{	"name": "NNXRuntimeORTDml", "skip": true, "comment": "DML metacommand only support 2 to 4 for inputs rank" 	},
+				{ "name": "NNXRuntimeORTDml", "skip": true, "comment": "DML metacommand only support 2 to 4 for inputs rank" },
 				{ "name": "NNXRuntimeHlsl", "relative_error": 0.01}
 			],
 			"datasets": [
@@ -593,6 +602,30 @@
 				  "outputs": [ { "shape": [ 4 ], "source": [ "0", "1", "0", "0" ] } ] }
 			]
 		},
+		{
+			"name": "Upsample.ScaleFromInput",
+			"runtimes": [	
+				{ "name": "NNXRuntimeDml", "skip": true, "comment": "Scale should be a constant tensor for RDG" },
+				{ "name": "NNXRuntimeHlsl", "skip": true, "comment": "Scale should be a constant tensor for RDG" }
+			],
+			"datasets": [
+				{ "inputs":  [ { "shape": [ 1 ] },                            { "shape": [ 1 ], "source": [ "3"    ],                  "type": "float" } ], "outputs": [ { "shape": [ 3 ]} ] },
+				{ "inputs":  [ { "shape": [ 2 ] },                            { "shape": [ 1 ], "source": [ "2"    ],                  "type": "float" } ], "outputs": [ { "shape": [ 4 ]} ] },
+				{ "inputs":  [ { "shape": [ 3 ] },                            { "shape": [ 1 ], "source": [ "1"    ],                  "type": "float" } ], "outputs": [ { "shape": [ 3 ]} ] },
+				{ "inputs":  [ { "shape": [ 1, 2, 3, 4 ] },                   { "shape": [ 4 ], "source": [ "1", "1", "2", "2" ],      "type": "float" } ], "outputs": [ { "shape": [ 1, 2, 6, 8 ]} ] },
+				{ "inputs":  [ { "shape": [ 2, 1, 3, 1, 2 ] },                { "shape": [ 5 ], "source": [ "1", "2", "2", "1", "2" ], "type": "float" } ], "outputs": [ { "shape": [ 2, 2, 6, 1, 4 ]} ] }
+			]
+		},
+		{
+			"name": "Upsample.ScaleFromWeight",
+			"datasets": [
+				{ "inputs":  [ { "shape": [ 1 ] } ],             "weights": [ { "shape": [ 1 ], "source": [ "3"    ],                  "type": "float" } ], "outputs": [ { "shape": [ 3 ]} ] },
+				{ "inputs":  [ { "shape": [ 2 ] } ],             "weights": [ { "shape": [ 1 ], "source": [ "2"    ],                  "type": "float" } ], "outputs": [ { "shape": [ 4 ]} ] },
+				{ "inputs":  [ { "shape": [ 3 ] } ],             "weights": [ { "shape": [ 1 ], "source": [ "1"    ],                  "type": "float" } ], "outputs": [ { "shape": [ 3 ]} ] },
+				{ "inputs":  [ { "shape": [ 1, 2, 3, 4 ] } ],    "weights": [ { "shape": [ 4 ], "source": [ "1", "1", "2", "2" ],      "type": "float" } ], "outputs": [ { "shape": [ 1, 2, 6, 8 ]} ] },
+				{ "inputs":  [ { "shape": [ 2, 1, 3, 1, 2 ] } ], "weights": [ { "shape": [ 5 ], "source": [ "1", "2", "2", "1", "2" ], "type": "float" } ], "outputs": [ { "shape": [ 2, 2, 6, 1, 4 ]} ] }
+			]
+		},
 		{
 			"name": "InstanceNormalization",
 			"datasets": [
diff --git a/Engine/Plugins/Experimental/NNE/Source/NNXQA/Resources/NNXQATestDescSchema.json b/Engine/Plugins/Experimental/NNE/Source/NNXQA/Resources/NNXQATestDescSchema.json
--- a/Engine/Plugins/Experimental/NNE/Source/NNXQA/Resources/NNXQATestDescSchema.json
+++ b/Engine/Plugins/Experimental/NNE/Source/NNXQA/Resources/NNXQATestDescSchema.json
@@ -95,7 +95,7 @@
 				},
 				"weights": {
 					"type": "array",
-					"description": "output to test",
+					"description": "weights to test",
 					"items": { "$ref": "#/$defs/test_config_tensor" }
 				},
 				"outputs": {
diff --git a/Engine/Plugins/Experimental/NNE/Source/NNXRuntimeRDG/Private/Hlsl/NNIRuntimeRDGUpsample.cpp b/Engine/Plugins/Experimental/NNE/Source/NNXRuntimeRDG/Private/Hlsl/NNIRuntimeRDGUpsample.cpp
new file mode 100644
--- /dev/null
+++ b/Engine/Plugins/Experimental/NNE/Source/NNXRuntimeRDG/Private/Hlsl/NNIRuntimeRDGUpsample.cpp
@@ -0,0 +1,150 @@
+// Copyright Epic Games, Inc. All Rights Reserved.
+
+#include "NNIRuntimeRDGUpsample.h"
+#include "NNEHlslShadersUpsampleCS.h"
+#include "NNXRuntimeHLSLHelper.h"
+
+namespace UE::NNIRuntimeRDG::Private::Hlsl
+{
+	DECLARE_GPU_STAT_NAMED(FNNIOperatorUpsample, TEXT("NNI.Operator.Hlsl.Upsample"));
+
+	/**
+	 * Upsample operator implementation
+	 */
+	class FUpsample : public NNX::FMLOperatorHlsl
+	{
+	public:
+
+		FUpsample() {}
+		virtual ~FUpsample() = default;
+
+	public:
+
+		virtual int PrepareOutputs(TConstArrayView<NNX::FTensorRef> InputTensors, TArrayView<NNX::FTensorRef> OutputTensors) const override
+		{
+			check(InputTensors.Num() == 2);
+			check(OutputTensors.Num() == 1);
+
+			const NNX::FTensor& X = *InputTensors[0];
+			const NNX::FTensor& Scales = *InputTensors[1];
+			
+			if (!Scales.HasPreparedData())
+			{
+				UE_LOG(LogNNX, Warning, TEXT("Upsample input 'Scale' (name: %s) should be constant for shape inference to succeed, however it is not."), *Scales.GetName());
+				return -1;
+			}
+
+			TConstArrayView<float> ScalesData = Scales.GetPreparedData<float>();
+
+			if (ScalesData.Num() != X.GetShape().Rank())
+			{
+				UE_LOG(LogNNX, Warning, TEXT("Upsample input 'Scale' (name: %s) have %d elements. While it should be the same as the rank of input 'X' (name : %s) witch is %d"), *Scales.GetName(), ScalesData.Num(), *X.GetName(), X.GetShape().Rank());
+				return -1;
+			}
+
+			NNX::FTensorShape OutputShape;
+			for (int32 i = 0; i < X.GetShape().Rank(); ++i)
+			{
+				
+				OutputShape.Data.Emplace(FMath::FloorToInt32(X.GetShape().Data[i] * ScalesData[i]));
+			}
+
+			OutputTensors[0]->SetShape(OutputShape);
+
+			return 0;
+		};
+
+		virtual bool Initialize(TConstArrayView<NNX::FTensorDesc> InputTensorDescs, TConstArrayView<NNX::FTensorDesc> OutputTensorDescs, const UE::NNECore::FAttributeMap& Attributes) override
+		{
+			check(InputTensorDescs.Num() == 2);
+			check(OutputTensorDescs.Num() == 1);
+
+			return true;
+		}
+
+		virtual void Dispatch(FRDGBuilder& GraphBuilder, TConstArrayView<NNX::FTensorRDGRef> InputTensors, TConstArrayView<NNX::FTensorRDGRef> OutputTensors) override
+		{
+			using namespace UE::NNEHlslShaders::Internal;
+			
+			check(InputTensors.Num() == 2);
+			check(OutputTensors.Num() == 1);
+			check(InputTensors[0] != nullptr);
+			check(InputTensors[1] != nullptr);
+			check(OutputTensors[0] != nullptr);
+			const NNX::FTensorRDG& Input = *InputTensors[0];
+			const NNX::FTensorRDG& Scales = *InputTensors[1];
+			const NNX::FTensorRDG& Output = *OutputTensors[0];
+
+			check(Scales.HasPreparedData());
+			TConstArrayView<float> ScalesData = Scales.GetPreparedData<float>();
+
+			FRDGBufferSRVRef InputSRV = GraphBuilder.CreateSRV(FRDGBufferSRVDesc(Input.GetBuffer(), PF_R32_FLOAT));
+			FRDGBufferUAVRef OutputUAV = GraphBuilder.CreateUAV(FRDGBufferUAVDesc(Output.GetBuffer(), PF_R32_FLOAT));
+
+			FIntVector ThreadGroupCount = NNX::ComputeElementWiseThreadGroups(Output.GetVolume(), FUpsampleConstants::NUM_GROUP_THREADS);
+
+			// Set parameters
+			FUpsampleCS::FParameters* Params = GraphBuilder.AllocParameters<FUpsampleCS::FParameters>();
+			Params->Input = InputSRV;
+			Params->Output = OutputUAV;
+			FillTensorStrideShaderParameters(Input, Params->TensorInfo, 0);
+			FillTensorStrideShaderParameters(Output, Params->TensorInfo, 1);
+			FillTensorSizeShaderParameters(Input, Params->TensorInfo, 2);
+			FillTensorSizeShaderParameters(Output, Params->TensorInfo, 3);
+			Params->Num = Output.GetVolume();
+			Params->ThreadCountX = ThreadGroupCount.X * FUpsampleConstants::NUM_GROUP_THREADS;
+
+			FUpsampleCS::FPermutationDomain PermutationVector;
+
+			PermutationVector.Set<FUpsampleCS::FUpsampleNumDimensions>(Output.GetShape().Rank());
+
+			TShaderMapRef<FUpsampleCS> ComputeShader(GetGlobalShaderMap(GMaxRHIFeatureLevel), PermutationVector);
+
+			RDG_EVENT_SCOPE(GraphBuilder, "NNI.Operator.Hlsl.Upsample");
+			RDG_GPU_STAT_SCOPE(GraphBuilder, FNNIOperatorUpsample);
+
+			FComputeShaderUtils::AddPass(
+				GraphBuilder,
+				RDG_EVENT_NAME("NNI.Operator.Hlsl.Upsample.Dispatch"),
+				ERDGPassFlags::Compute | ERDGPassFlags::NeverCull,
+				ComputeShader,
+				Params,
+				ThreadGroupCount);
+		}
+	};
+
+	bool ValidateUpsampleOperator(const UE::NNECore::FAttributeMap& AttributeMap, TConstArrayView<EMLTensorDataType> InputTypes, TConstArrayView<NNX::FSymbolicTensorShape> InputShapes)
+	{
+		bool bIsValid = true;
+
+		NNX::FAttributeValidator AttributeValidator;
+		AttributeValidator.AddOptional(TEXT("mode"), ENNEAttributeDataType::String);
+		bIsValid &= AttributeValidator.Validate(AttributeMap);
+		
+		FString Mode = AttributeMap.GetValueOrDefault<FString>(TEXT("mode"), TEXT("nearest"));
+		if (!Mode.Equals(TEXT("nearest")))
+		{
+			UE_LOG(LogNNX, Warning, TEXT("Upsample HLSL operator only supports nearest mode for now"));
+			return false;
+		}
+
+		NNX::FInputValidator InputValidator;
+		InputValidator.AddSupportedType(EMLTensorDataType::Float);
+		InputValidator.AddRequired();
+		InputValidator.AddRequired();
+		bIsValid &= InputValidator.Validate(InputTypes);
+
+		return bIsValid;
+	}
+
+	NNX::FMLOperatorHlsl* CreateUpsampleOperator()
+	{
+		return new FUpsample();
+	}
+
+	bool RegisterUpsampleOperator(NNX::FMLOperatorRegistryHlsl& Registry)
+	{
+		Registry.OpAdd(TEXT("Upsample"), CreateUpsampleOperator, ValidateUpsampleOperator);
+		return true;
+	}
+} // UE::NNIRuntimeRDG::Private::Hlsl
diff --git a/Engine/Plugins/Experimental/NNE/Source/NNXRuntimeRDG/Private/Hlsl/NNIRuntimeRDGUpsample.h b/Engine/Plugins/Experimental/NNE/Source/NNXRuntimeRDG/Private/Hlsl/NNIRuntimeRDGUpsample.h
new file mode 100644
--- /dev/null
+++ b/Engine/Plugins/Experimental/NNE/Source/NNXRuntimeRDG/Private/Hlsl/NNIRuntimeRDGUpsample.h
@@ -0,0 +1,11 @@
+// Copyright Epic Games, Inc. All Rights Reserved.
+
+#pragma once
+
+#include "../NNXRuntimeRDG.h"
+#include "../NNXRuntimeHLSLOp.h"
+
+namespace UE::NNIRuntimeRDG::Private::Hlsl
+{
+	bool RegisterUpsampleOperator(NNX::FMLOperatorRegistryHlsl& Registry);
+} // UE::NNIRuntimeRDG::Private::Hlsl
diff --git a/Engine/Plugins/Experimental/NNE/Source/NNXRuntimeRDG/Private/NNXRuntimeHLSL.cpp b/Engine/Plugins/Experimental/NNE/Source/NNXRuntimeRDG/Private/NNXRuntimeHLSL.cpp
--- a/Engine/Plugins/Experimental/NNE/Source/NNXRuntimeRDG/Private/NNXRuntimeHLSL.cpp
+++ b/Engine/Plugins/Experimental/NNE/Source/NNXRuntimeRDG/Private/NNXRuntimeHLSL.cpp
@@ -13,6 +13,7 @@
 #include "Hlsl/NNIRuntimeRDGGemm.h"
 #include "Hlsl/NNIRuntimeRDGMatMul.h"
 #include "Hlsl/NNIRuntimeRDGInstanceNormalization.h"
+#include "Hlsl/NNIRuntimeRDGUpsample.h"
 #include "NNXRuntimeHLSLHelper.h"
 #include "NNXRuntimeRDG.h"
 
@@ -266,6 +267,7 @@ public:
 		UE::NNIRuntimeRDG::Private::Hlsl::RegisterConvTransposeOperator(*registry);
 		UE::NNIRuntimeRDG::Private::Hlsl::RegisterMatMulOperator(*registry);
 		UE::NNIRuntimeRDG::Private::Hlsl::RegisterInstanceNormalizationOperator(*registry);
+		UE::NNIRuntimeRDG::Private::Hlsl::RegisterUpsampleOperator(*registry);
 
 		return true;
 	}
diff --git a/Engine/Plugins/Experimental/NNE/Source/NNXRuntimeRDG/Private/NNXRuntimeHLSLHelper.cpp b/Engine/Plugins/Experimental/NNE/Source/NNXRuntimeRDG/Private/NNXRuntimeHLSLHelper.cpp
--- a/Engine/Plugins/Experimental/NNE/Source/NNXRuntimeRDG/Private/NNXRuntimeHLSLHelper.cpp
+++ b/Engine/Plugins/Experimental/NNE/Source/NNXRuntimeRDG/Private/NNXRuntimeHLSLHelper.cpp
@@ -6,14 +6,23 @@
 
 namespace NNX
 {
-	void FillTensorStrideShaderParameters(const FTensor& TensorDesc, TStaticArray<FUintVector4, NXRT_TENSORSTRIDEINFO_MAX_NUM_DIMENSIONS, 16U>& OutShaderParam, int32 Idx, int32 TargetNumdimensionForBroadcast)
+	void FillTensorSizeShaderParameters(const FTensor& Tensor, TStaticArray<FUintVector4, NXRT_TENSORSTRIDEINFO_MAX_NUM_DIMENSIONS, 16U>& OutShaderParam, int32 Idx)
+	{
+		static_assert(FTensorShape::MaxRank <= NXRT_TENSORSTRIDEINFO_MAX_NUM_DIMENSIONS);
+		for (int32 i = 0; i < Tensor.GetShape().Rank(); ++i)
+		{
+			OutShaderParam[i][Idx] = Tensor.GetShape().Data[i];
+		}
+	}
+
+	void FillTensorStrideShaderParameters(const FTensor& Tensor, TStaticArray<FUintVector4, NXRT_TENSORSTRIDEINFO_MAX_NUM_DIMENSIONS, 16U>& OutShaderParam, int32 Idx, int32 TargetNumdimensionForBroadcast)
 	{
 		if (TargetNumdimensionForBroadcast == -1)
 		{
-			TargetNumdimensionForBroadcast = TensorDesc.GetShape().Rank();
+			TargetNumdimensionForBroadcast = Tensor.GetShape().Rank();
 		}
-		checkf(TargetNumdimensionForBroadcast >= TensorDesc.GetShape().Rank(), TEXT("Can't broadcast tensor from rank %d to rank %d, should be inferior or equal."), TensorDesc.GetShape().Rank(), TargetNumdimensionForBroadcast);
-		int32 Offset = TargetNumdimensionForBroadcast - TensorDesc.GetShape().Rank();
+		checkf(TargetNumdimensionForBroadcast >= Tensor.GetShape().Rank(), TEXT("Can't broadcast tensor from rank %d to rank %d, should be inferior or equal."), Tensor.GetShape().Rank(), TargetNumdimensionForBroadcast);
+		int32 Offset = TargetNumdimensionForBroadcast - Tensor.GetShape().Rank();
 
 		static_assert(FTensorShape::MaxRank <= NXRT_TENSORSTRIDEINFO_MAX_NUM_DIMENSIONS);
 		for (int32 i = 7; i >= 0; --i)
@@ -28,20 +37,20 @@ namespace NNX
 			}
 			else
 			{
-				OutShaderParam[i][Idx] = OutShaderParam[i + 1][Idx] * TensorDesc.GetShape().Data[i + 1 - Offset];
+				OutShaderParam[i][Idx] = OutShaderParam[i + 1][Idx] * Tensor.GetShape().Data[i + 1 - Offset];
 			}
 		}
 	}
 
-	void FillTensorStrideForBroadcastShaderParameters(const FTensor& TensorDesc, int32 OutputNumdimension, TStaticArray<FUintVector4, NXRT_TENSORSTRIDEINFO_MAX_NUM_DIMENSIONS, 16U>& OutShaderParam, int32 Idx)
+	void FillTensorStrideForBroadcastShaderParameters(const FTensor& Tensor, int32 OutputNumdimension, TStaticArray<FUintVector4, NXRT_TENSORSTRIDEINFO_MAX_NUM_DIMENSIONS, 16U>& OutShaderParam, int32 Idx)
 	{
-		checkf(OutputNumdimension >= TensorDesc.GetShape().Rank(), TEXT("Can't broadcast tensor from rank %d to rank %d, should be inferior or equal."), TensorDesc.GetShape().Rank(), OutputNumdimension);
-		FillTensorStrideShaderParameters(TensorDesc, OutShaderParam, Idx, OutputNumdimension);
-		int32 Offset = OutputNumdimension - TensorDesc.GetShape().Rank();
+		checkf(OutputNumdimension >= Tensor.GetShape().Rank(), TEXT("Can't broadcast tensor from rank %d to rank %d, should be inferior or equal."), Tensor.GetShape().Rank(), OutputNumdimension);
+		FillTensorStrideShaderParameters(Tensor, OutShaderParam, Idx, OutputNumdimension);
+		int32 Offset = OutputNumdimension - Tensor.GetShape().Rank();
 		for (int32 i = Offset; i < OutputNumdimension; ++i)
 		{
 			// the stride for broadcast dimension is kept as 0
-			if (TensorDesc.GetShape().Data[i - Offset] == 1)
+			if (Tensor.GetShape().Data[i - Offset] == 1)
 			{
 				OutShaderParam[i][Idx] = 0;
 			}
diff --git a/Engine/Plugins/Experimental/NNE/Source/NNXRuntimeRDG/Private/NNXRuntimeHLSLHelper.h b/Engine/Plugins/Experimental/NNE/Source/NNXRuntimeRDG/Private/NNXRuntimeHLSLHelper.h
--- a/Engine/Plugins/Experimental/NNE/Source/NNXRuntimeRDG/Private/NNXRuntimeHLSLHelper.h
+++ b/Engine/Plugins/Experimental/NNE/Source/NNXRuntimeRDG/Private/NNXRuntimeHLSLHelper.h
@@ -8,7 +8,8 @@
 
 namespace NNX
 {
-	void FillTensorStrideShaderParameters(const FTensor& TensorDesc, TStaticArray<FUintVector4, NXRT_TENSORSTRIDEINFO_MAX_NUM_DIMENSIONS, 16U>& OutShaderParam, int32 Idx, int32 TargetNumdimensionForBroadcast = -1);
-	void FillTensorStrideForBroadcastShaderParameters(const FTensor& TensorDesc, int32 OutputNumdimension, TStaticArray<FUintVector4, NXRT_TENSORSTRIDEINFO_MAX_NUM_DIMENSIONS, 16U>& OutShaderParam, int32 Idx);
+	void FillTensorSizeShaderParameters(const FTensor& Tensor, TStaticArray<FUintVector4, NXRT_TENSORSTRIDEINFO_MAX_NUM_DIMENSIONS, 16U>& OutShaderParam, int32 Idx);
+	void FillTensorStrideShaderParameters(const FTensor& Tensor, TStaticArray<FUintVector4, NXRT_TENSORSTRIDEINFO_MAX_NUM_DIMENSIONS, 16U>& OutShaderParam, int32 Idx, int32 TargetNumdimensionForBroadcast = -1);
+	void FillTensorStrideForBroadcastShaderParameters(const FTensor& Tensor, int32 OutputNumdimension, TStaticArray<FUintVector4, NXRT_TENSORSTRIDEINFO_MAX_NUM_DIMENSIONS, 16U>& OutShaderParam, int32 Idx);
 	FIntVector ComputeElementWiseThreadGroups(uint32 ElementCount, uint32 GroupSizeX);
 } // NNX
\ No newline at end of file
diff --git a/Engine/Plugins/Experimental/NNE/Source/NNXRuntimeRDG/Private/NNXRuntimeRDG.cpp b/Engine/Plugins/Experimental/NNE/Source/NNXRuntimeRDG/Private/NNXRuntimeRDG.cpp
--- a/Engine/Plugins/Experimental/NNE/Source/NNXRuntimeRDG/Private/NNXRuntimeRDG.cpp
+++ b/Engine/Plugins/Experimental/NNE/Source/NNXRuntimeRDG/Private/NNXRuntimeRDG.cpp
@@ -498,15 +498,6 @@ int FMLInferenceModelRDG::EnqueueRDG(FRDGBuilder& RDGBuilder, TConstArrayView<FM
 		TensorRDG.SetBuffer(TensorBuffer);
 	}
 
-	//For now weights tensors are not uploaded to GPU thus GetBuffer will return nullptr for them.
-	checkCode(for (const FTensorRDG* TensorRDG : AllTensorRDGs) { if (TensorRDG != nullptr) { check(TensorRDG->GetBuffer() != nullptr); } });
-
-	//Insert weights tensors
-	for (int32 i = 0; i < WeightTensorIndices.Num(); ++i)
-	{
-		AllTensorRDGs[WeightTensorIndices[i]] = &WeightTensorRDGs[i];
-	}
-	
 	// We can now dispatch operators
 	AddDispatchOps_RenderThread(RDGBuilder);
 
diff --git a/Engine/Plugins/Experimental/NNE/Source/NNXUtils/Private/NNXModelBuilderONNX.cpp b/Engine/Plugins/Experimental/NNE/Source/NNXUtils/Private/NNXModelBuilderONNX.cpp
--- a/Engine/Plugins/Experimental/NNE/Source/NNXUtils/Private/NNXModelBuilderONNX.cpp
+++ b/Engine/Plugins/Experimental/NNE/Source/NNXUtils/Private/NNXModelBuilderONNX.cpp
@@ -332,7 +332,14 @@ NNXUTILS_API bool CreateONNXModelForOperator(bool UseVariadicShapeForModel, cons
 {
 	Model = FNNIModelRaw{};
 	
-	TUniquePtr<IMLModelBuilder> Builder(CreateONNXModelBuilder());
+	int64 IrVersion = OnnxIrVersion;
+	int64 OpsetVersion = OnnxOpsetVersion;
+	if (OperatorName == TEXT("Upsample"))
+	{
+		// Upsample is deprecated starting opset 10
+		OpsetVersion = 9;
+	}
+	TUniquePtr<IMLModelBuilder> Builder(CreateONNXModelBuilder(IrVersion, OpsetVersion));
 
 	Builder->Begin();
 
