NNE Support for Pad operator in HLSL

Problem:
* Pad opset 9 needed for style transfer was not implemented

Solution:
* Added pad in onnx opset 9 version

Tests:
* All NNX tests pass under win64 editor + added/updated pad tests

diff --git a/Engine/Plugins/Experimental/NNE/Source/NNEHlslShaders/Internal/NNEHlslShadersPadCS.h b/Engine/Plugins/Experimental/NNE/Source/NNEHlslShaders/Internal/NNEHlslShadersPadCS.h
new file mode 100644
--- /dev/null
+++ b/Engine/Plugins/Experimental/NNE/Source/NNEHlslShaders/Internal/NNEHlslShadersPadCS.h
@@ -0,0 +1,51 @@
+// Copyright Epic Games, Inc. All Rights Reserved.
+
+#pragma once
+
+#include "GlobalShader.h"
+#include "ShaderParameterUtils.h"
+#include "RenderGraphUtils.h"
+#include "NNXOperator.h"
+
+namespace UE::NNEHlslShaders::Internal
+{
+	enum class EPadMode : uint8
+	{
+		CONSTANT = 0,
+		REFLECT,
+		EDGE,
+		MAX
+	};
+	
+	class FPadConstants
+	{
+	public:
+
+		static const int32 MAX_NUM_DIMENSIONS{ 8 };
+		static const int32 NUM_GROUP_THREADS{ 256 };
+	};
+
+	class NNEHLSLSHADERS_API FPadCS : public FGlobalShader
+	{
+		DECLARE_GLOBAL_SHADER(FPadCS);
+		SHADER_USE_PARAMETER_STRUCT(FPadCS, FGlobalShader)
+
+		class FPadMode : SHADER_PERMUTATION_ENUM_CLASS("MODE", EPadMode);
+		class FPadNumDimensions : SHADER_PERMUTATION_RANGE_INT("NUM_DIMENSIONS", 1, FPadConstants::MAX_NUM_DIMENSIONS);
+		using FPermutationDomain = TShaderPermutationDomain<FPadMode, FPadNumDimensions>;
+
+	public:
+
+		BEGIN_SHADER_PARAMETER_STRUCT(FParameters, )
+			SHADER_PARAMETER_RDG_BUFFER_SRV(Buffer<float>, Input)
+			SHADER_PARAMETER_RDG_BUFFER_UAV(RWBuffer<float>, Output)
+			SHADER_PARAMETER_ARRAY(FUintVector4, TensorInfo, [FPadConstants::MAX_NUM_DIMENSIONS])
+			SHADER_PARAMETER(float, Value)
+			SHADER_PARAMETER(uint32, Num)
+			SHADER_PARAMETER(uint32, ThreadCountX)
+		END_SHADER_PARAMETER_STRUCT()
+
+		static void ModifyCompilationEnvironment(const FGlobalShaderPermutationParameters& InParameters, FShaderCompilerEnvironment& OutEnvironment);
+		static void LexFromString(EPadMode& OutValue, const TCHAR* StringVal);
+	};
+} // UE::NNEHlslShaders::Internal
diff --git a/Engine/Plugins/Experimental/NNE/Source/NNEHlslShaders/Private/NNEHlslShadersPadCS.cpp b/Engine/Plugins/Experimental/NNE/Source/NNEHlslShaders/Private/NNEHlslShadersPadCS.cpp
new file mode 100644
--- /dev/null
+++ b/Engine/Plugins/Experimental/NNE/Source/NNEHlslShaders/Private/NNEHlslShadersPadCS.cpp
@@ -0,0 +1,26 @@
+// Copyright Epic Games, Inc. All Rights Reserved.
+
+#include "NNEHlslShadersPadCS.h"
+#include "NNXCore.h"
+
+namespace UE::NNEHlslShaders::Internal
+{
+	void FPadCS::ModifyCompilationEnvironment(const FGlobalShaderPermutationParameters& InParameters, FShaderCompilerEnvironment& OutEnvironment)
+	{
+		FGlobalShader::ModifyCompilationEnvironment(InParameters, OutEnvironment);
+
+		OutEnvironment.SetDefine(TEXT("THREADGROUP_SIZE_X"), FPadConstants::NUM_GROUP_THREADS);
+
+		FPermutationDomain PermutationVector(InParameters.PermutationId);
+	}
+
+	void FPadCS::LexFromString(EPadMode& OutValue, const TCHAR* StringVal)
+	{
+		OutValue = EPadMode::CONSTANT;
+		if (FCString::Stricmp(StringVal, TEXT("CONSTANT")) == 0) OutValue = EPadMode::CONSTANT;
+		else if (FCString::Stricmp(StringVal, TEXT("REFLECT")) == 0) OutValue = EPadMode::REFLECT;
+		else if (FCString::Stricmp(StringVal, TEXT("EDGE")) == 0) OutValue = EPadMode::EDGE;
+	}
+
+	IMPLEMENT_GLOBAL_SHADER(FPadCS, "/NNE/NNEHlslShadersPad.usf", "Pad", SF_Compute);
+} // UE::NNEHlslShaders::Internal
\ No newline at end of file
diff --git a/Engine/Plugins/Experimental/NNE/Source/NNEHlslShaders/Shaders/NNEHlslShadersPad.usf b/Engine/Plugins/Experimental/NNE/Source/NNEHlslShaders/Shaders/NNEHlslShadersPad.usf
new file mode 100644
--- /dev/null
+++ b/Engine/Plugins/Experimental/NNE/Source/NNEHlslShaders/Shaders/NNEHlslShadersPad.usf
@@ -0,0 +1,70 @@
+// Copyright Epic Games, Inc. All Rights Reserved.
+
+#include "/Engine/Public/Platform.ush"
+#include "NNEHlslShadersBroadcastHelper.ush"
+
+Buffer<float> Input;
+RWBuffer<float> Output;
+uint4 TensorInfo[NUM_DIMENSIONS];
+float Value;
+uint Num;
+uint ThreadCountX;
+
+#define INPUT_STRIDE 0
+#define OUTPUT_STRIDE 1
+#define INPUT_SIZE 2
+#define PRE_PAD 3
+
+// Must correspond to EPadMode defined in NNEHlslShadersPadCS.h
+#define CONSTANT_MODE 0
+#define REFLECT_MODE 1
+#define EDGE_MODE 2
+
+[numthreads(THREADGROUP_SIZE_X, 1, 1)]
+void Pad(in const uint3 DispatchThreadID : SV_DispatchThreadID)
+{
+	const uint Index = DispatchThreadID.y * ThreadCountX + DispatchThreadID.x;
+	
+	if (Index < Num)
+	{
+		uint InputIndex = 0;
+
+		uint Offset = Index;
+		for (uint dim = 0; dim < NUM_DIMENSIONS; ++dim)
+		{
+			uint OutDimIdx, R;
+			DivMod(Offset, TensorInfo[dim][OUTPUT_STRIDE], OutDimIdx, R);
+			
+			int PrePad = TensorInfo[dim][PRE_PAD];
+			int InputDimIdx = (int)OutDimIdx - PrePad;
+			int InputDimSize = (int) TensorInfo[dim][INPUT_SIZE];
+			
+			if (InputDimIdx < 0 || InputDimIdx >= InputDimSize)
+			{
+				#if MODE == CONSTANT_MODE
+				Output[Index] = Value;
+				return;
+				#elif MODE == REFLECT_MODE
+				if (InputDimIdx >= InputDimSize)
+				{
+					int OverCnt = InputDimIdx - (InputDimSize - 2);
+					InputDimIdx = InputDimSize - OverCnt;
+				}
+				else if (InputDimIdx < 0)
+				{
+					InputDimIdx = -InputDimIdx;
+				}
+				#endif
+				//MODE == EDGE_MODE + clamping large reflect indices
+				InputDimIdx = min(InputDimIdx, InputDimSize - 1);
+				InputDimIdx = max(InputDimIdx, 0);
+			}
+			
+			InputIndex += TensorInfo[dim][INPUT_STRIDE] * InputDimIdx;
+			Offset = R;
+		}
+
+		float X = Input[InputIndex];
+		Output[Index] = X;
+	}
+}
diff --git a/Engine/Plugins/Experimental/NNE/Source/NNXQA/Resources/NNXQATestDesc.json b/Engine/Plugins/Experimental/NNE/Source/NNXQA/Resources/NNXQATestDesc.json
--- a/Engine/Plugins/Experimental/NNE/Source/NNXQA/Resources/NNXQATestDesc.json
+++ b/Engine/Plugins/Experimental/NNE/Source/NNXQA/Resources/NNXQATestDesc.json
@@ -296,8 +296,7 @@
 		{
 			"category": "PadOp",
 			"runtimes": [
-				{ "name": "NNXRuntimeDml", "skip": true	},
-				{ "name": "NNXRuntimeHlsl", "skip": true	}
+				{ "name": "NNXRuntimeDml", "skip": true	}
 			],
 			"targets": [
 				{ "target": "Pad" }
@@ -587,19 +586,33 @@
 			]
 		},
 		{
-			"name": "PadOp.PaddingFromInput",
+			"name": "PadOp.Pad00Rank1",
 			"datasets": [
-				{ "inputs":  [ { "shape": [ 1 ], "source": [ "1" ] } , 
-							   { "shape": [ 2 ], "source": [ "1", "2" ] , "type": "int64" } ],
-				  "outputs": [ { "shape": [ 4 ], "source": [ "0", "1", "0", "0" ] } ] }
+				{ "inputs":  [ { "shape": [ 3 ] } ], "outputs": [ { "shape": [ 3 ] } ] }
 			]
 		},
 		{
-			"name": "PadOp.PaddingFromWeight",
+			"name": "PadOp.Pad03Rank1",
 			"datasets": [
-				{ "inputs":  [ { "shape": [ 1 ], "source": [ "1" ] } ], 
-				  "weights": [ { "shape": [ 2 ], "source": [ "1", "2" ] , "type": "int64" } ],
-				  "outputs": [ { "shape": [ 4 ], "source": [ "0", "1", "0", "0" ] } ] }
+				{ "inputs":  [ { "shape": [ 4 ] } ], "outputs": [ { "shape": [ 7 ] } ] }
+			]
+		},
+		{
+			"name": "PadOp.Pad20Rank1",
+			"datasets": [
+				{ "inputs":  [ { "shape": [ 4 ] } ], "outputs": [ { "shape": [ 6 ] } ] }
+			]
+		},
+		{
+			"name": "PadOp.Pad41Rank1",
+			"datasets": [
+				{ "inputs":  [ { "shape": [ 5 ] } ], "outputs": [ { "shape": [ 10 ] } ] }
+			]
+		},
+		{
+			"name": "PadOp.Pad00_10_12_34_Rank4",
+			"datasets": [
+				{ "inputs":  [ { "shape": [ 1,32,224,128 ] } ], "outputs": [ { "shape": [ 1,33,227,135 ] } ] }
 			]
 		},
 		{
@@ -612,7 +625,7 @@
 				{ "inputs":  [ { "shape": [ 1 ] },                            { "shape": [ 1 ], "source": [ "3"    ],                  "type": "float" } ], "outputs": [ { "shape": [ 3 ]} ] },
 				{ "inputs":  [ { "shape": [ 2 ] },                            { "shape": [ 1 ], "source": [ "2"    ],                  "type": "float" } ], "outputs": [ { "shape": [ 4 ]} ] },
 				{ "inputs":  [ { "shape": [ 3 ] },                            { "shape": [ 1 ], "source": [ "1"    ],                  "type": "float" } ], "outputs": [ { "shape": [ 3 ]} ] },
-				{ "inputs":  [ { "shape": [ 1, 2, 3, 4 ] },                   { "shape": [ 4 ], "source": [ "1", "1", "2", "2" ],      "type": "float" } ], "outputs": [ { "shape": [ 1, 2, 6, 8 ]} ] },
+				{ "inputs":  [ { "shape": [ 1, 128, 50, 55 ] },               { "shape": [ 4 ], "source": [ "1", "1", "2", "2" ],      "type": "float" } ], "outputs": [ { "shape": [ 1, 128, 100, 110 ]} ] },
 				{ "inputs":  [ { "shape": [ 2, 1, 3, 1, 2 ] },                { "shape": [ 5 ], "source": [ "1", "2", "2", "1", "2" ], "type": "float" } ], "outputs": [ { "shape": [ 2, 2, 6, 1, 4 ]} ] }
 			]
 		},
@@ -622,7 +635,7 @@
 				{ "inputs":  [ { "shape": [ 1 ] } ],             "weights": [ { "shape": [ 1 ], "source": [ "3"    ],                  "type": "float" } ], "outputs": [ { "shape": [ 3 ]} ] },
 				{ "inputs":  [ { "shape": [ 2 ] } ],             "weights": [ { "shape": [ 1 ], "source": [ "2"    ],                  "type": "float" } ], "outputs": [ { "shape": [ 4 ]} ] },
 				{ "inputs":  [ { "shape": [ 3 ] } ],             "weights": [ { "shape": [ 1 ], "source": [ "1"    ],                  "type": "float" } ], "outputs": [ { "shape": [ 3 ]} ] },
-				{ "inputs":  [ { "shape": [ 1, 2, 3, 4 ] } ],    "weights": [ { "shape": [ 4 ], "source": [ "1", "1", "2", "2" ],      "type": "float" } ], "outputs": [ { "shape": [ 1, 2, 6, 8 ]} ] },
+				{ "inputs":  [ { "shape": [ 1, 128, 50, 55 ] } ],    "weights": [ { "shape": [ 4 ], "source": [ "1", "1", "2", "2" ],      "type": "float" } ], "outputs": [ { "shape": [ 1, 128, 100, 110 ]} ] },
 				{ "inputs":  [ { "shape": [ 2, 1, 3, 1, 2 ] } ], "weights": [ { "shape": [ 5 ], "source": [ "1", "2", "2", "1", "2" ], "type": "float" } ], "outputs": [ { "shape": [ 2, 2, 6, 1, 4 ]} ] }
 			]
 		},
@@ -961,7 +974,63 @@
 				{ "attributes": [ { "name": "__UE__algorithm", "value": "SharedMemory16x16", "type": "String" } ] },
 				{ "attributes": [ { "name": "__UE__algorithm", "value": "SharedMemory32x8", "type": "String" } ] }
 			]
+		},
+		{
+			"name": "PadOp.ModeAttributeSet",
+			"attribute_maps": [
+				{ "attributes": [ { "name": "mode", "value": "constant", "type": "String" } ] },
+				{ "attributes": [ { "name": "mode", "value": "reflect", "type": "String" } ] },
+				{ "attributes": [ { "name": "mode", "value": "edge", "type": "String" } ] }
+			]
+		},
+		{
+			"name": "PadOp.Pad00Rank1",
+			"multiply_with_attribute_sets": [ "PadOp.ModeAttributeSet" ],
+			"attribute_maps": [
+				{ "attributes": [
+					{ "name": "value", "value": 1, "type": "float" },
+					{ "name": "pads", "value": [ 0, 0 ], "type": "Int32Array" }
+				] }
+			]
+		},
+		{
+			"name": "PadOp.Pad03Rank1",
+			"multiply_with_attribute_sets": [ "PadOp.ModeAttributeSet" ],
+			"attribute_maps": [
+				{ "attributes": [
+					{ "name": "value", "value": 2, "type": "float" },
+					{ "name": "pads", "value": [ 0, 3 ], "type": "Int32Array" }
+				] }
+			]
+		},
+		{
+			"name": "PadOp.Pad20Rank1",
+			"multiply_with_attribute_sets": [ "PadOp.ModeAttributeSet" ],
+			"attribute_maps": [
+				{ "attributes": [
+					{ "name": "value", "value": 3, "type": "float" },
+					{ "name": "pads", "value": [ 2, 0 ], "type": "Int32Array" }
+				] }
+			]
+		},
+		{
+			"name": "PadOp.Pad41Rank1",
+			"multiply_with_attribute_sets": [ "PadOp.ModeAttributeSet" ],
+			"attribute_maps": [
+				{ "attributes": [
+					{ "name": "pads", "value": [ 4, 1 ], "type": "Int32Array" }
+				] }
+			]
+		},
+		{
+			"name": "PadOp.Pad00_10_12_34_Rank4",
+			"multiply_with_attribute_sets": [ "PadOp.ModeAttributeSet" ],
+			"attribute_maps": [
+				{ "attributes": [
+					{ "name": "pads", "value": [ 0, 1, 1, 3, 0, 0, 2, 4 ], "type": "Int32Array" }
+				] }
+			]
 		}
 	]
-	
+
 }
diff --git a/Engine/Plugins/Experimental/NNE/Source/NNXRuntimeRDG/Private/Hlsl/NNIRuntimeRDGPad.cpp b/Engine/Plugins/Experimental/NNE/Source/NNXRuntimeRDG/Private/Hlsl/NNIRuntimeRDGPad.cpp
new file mode 100644
--- /dev/null
+++ b/Engine/Plugins/Experimental/NNE/Source/NNXRuntimeRDG/Private/Hlsl/NNIRuntimeRDGPad.cpp
@@ -0,0 +1,164 @@
+// Copyright Epic Games, Inc. All Rights Reserved.
+
+#include "NNIRuntimeRDGPad.h"
+#include "NNEHlslShadersPadCS.h"
+#include "NNXRuntimeHLSLHelper.h"
+
+namespace UE::NNIRuntimeRDG::Private::Hlsl
+{
+	DECLARE_GPU_STAT_NAMED(FNNIOperatorPad, TEXT("NNI.Operator.Hlsl.Pad"));
+
+	using EPadMode = UE::NNEHlslShaders::Internal::EPadMode;
+
+	/**
+	 * Pad operator implementation
+	 */
+	class FPad : public NNX::FMLOperatorHlsl
+	{
+	public:
+
+		FPad() {}
+		virtual ~FPad() = default;
+
+		TArray<int32> Pads;
+		float Value;
+		EPadMode Mode;
+
+	public:
+
+		virtual int PrepareOutputs(TConstArrayView<NNX::FTensorRef> InputTensors, TArrayView<NNX::FTensorRef> OutputTensors) const override
+		{
+			check(InputTensors.Num() == 1);
+			check(OutputTensors.Num() == 1);
+
+			const NNX::FTensor& X = *InputTensors[0];
+
+			NNX::FTensorShape OutputShape;
+			for (int32 i = 0; i < X.GetShape().Rank(); ++i)
+			{
+				int32 PrePad = Pads[i];
+				int32 PostPad = Pads[i + X.GetShape().Rank()];
+				int32 OutputDim = PrePad + X.GetShape().Data[i] + PostPad;
+				if (OutputDim < 1)
+				{
+					UE_LOG(LogNNX, Warning, TEXT("Pads cannot reduce dimension below 1, but would for tensor (name:%s) at rank %d of size %d with prepad %d and postpad %d."), *X.GetName(), i, X.GetShape().Data[i], PrePad, PostPad);
+					return -1;
+				}
+				OutputShape.Data.Emplace(OutputDim);
+			}
+
+			OutputTensors[0]->SetShape(OutputShape);
+
+			return 0;
+		};
+
+		virtual bool Initialize(TConstArrayView<NNX::FTensorDesc> InputTensorDescs, TConstArrayView<NNX::FTensorDesc> OutputTensorDescs, const UE::NNECore::FAttributeMap& Attributes) override
+		{
+			using namespace UE::NNEHlslShaders::Internal;
+			
+			check(InputTensorDescs.Num() == 1);
+			check(OutputTensorDescs.Num() == 1);
+			
+			Pads = Attributes.GetValue<TArray<int32>>(TEXT("pads"));
+			Value = Attributes.GetValueOrDefault<float>(TEXT("value"), 0.0f);
+			FPadCS::LexFromString(Mode, *Attributes.GetValue<FString>(TEXT("mode")));
+
+			if ((2*InputTensorDescs[0].GetShape().Rank()) != Pads.Num())
+			{
+				UE_LOG(LogNNX, Warning, TEXT("pads attribute lenght (%d) should be twice the rank of input X (%d)."), Pads.Num(), InputTensorDescs[0].GetShape().Rank());
+				return false;
+			}
+
+			return true;
+		}
+
+		virtual void Dispatch(FRDGBuilder& GraphBuilder, TConstArrayView<NNX::FTensorRDGRef> InputTensors, TConstArrayView<NNX::FTensorRDGRef> OutputTensors) override
+		{
+			using namespace UE::NNEHlslShaders::Internal;
+			
+			check(InputTensors.Num() == 1);
+			check(OutputTensors.Num() == 1);
+			check(InputTensors[0] != nullptr);
+			check(OutputTensors[0] != nullptr);
+			const NNX::FTensorRDG& Input = *InputTensors[0];
+			const NNX::FTensorRDG& Output = *OutputTensors[0];
+
+			FRDGBufferSRVRef InputSRV = GraphBuilder.CreateSRV(FRDGBufferSRVDesc(Input.GetBuffer(), PF_R32_FLOAT));
+			FRDGBufferUAVRef OutputUAV = GraphBuilder.CreateUAV(FRDGBufferUAVDesc(Output.GetBuffer(), PF_R32_FLOAT));
+
+			FIntVector ThreadGroupCount = NNX::ComputeElementWiseThreadGroups(Output.GetVolume(), FPadConstants::NUM_GROUP_THREADS);
+
+			// Set parameters
+			FPadCS::FParameters* Params = GraphBuilder.AllocParameters<FPadCS::FParameters>();
+			Params->Input = InputSRV;
+			Params->Output = OutputUAV;
+			FillTensorStrideShaderParameters(Input, Params->TensorInfo, 0);
+			FillTensorStrideShaderParameters(Output, Params->TensorInfo, 1);
+			FillTensorSizeShaderParameters(Input, Params->TensorInfo, 2);
+			for (int32 i = 0; i < Input.GetShape().Rank(); ++i)
+			{
+				Params->TensorInfo[i][3] = Pads[i];//Pre-pad
+			}
+			Params->Value = Value;
+			Params->Num = Output.GetVolume();
+			Params->ThreadCountX = ThreadGroupCount.X * FPadConstants::NUM_GROUP_THREADS;
+
+			FPadCS::FPermutationDomain PermutationVector;
+			PermutationVector.Set<FPadCS::FPadMode>(Mode);
+			PermutationVector.Set<FPadCS::FPadNumDimensions>(Output.GetShape().Rank());
+
+			TShaderMapRef<FPadCS> ComputeShader(GetGlobalShaderMap(GMaxRHIFeatureLevel), PermutationVector);
+
+			RDG_EVENT_SCOPE(GraphBuilder, "NNI.Operator.Hlsl.Pad");
+			RDG_GPU_STAT_SCOPE(GraphBuilder, FNNIOperatorPad);
+
+			FComputeShaderUtils::AddPass(
+				GraphBuilder,
+				RDG_EVENT_NAME("NNI.Operator.Hlsl.Pad.Dispatch"),
+				ERDGPassFlags::Compute | ERDGPassFlags::NeverCull,
+				ComputeShader,
+				Params,
+				ThreadGroupCount);
+		}
+	};
+
+	bool ValidatePadOperator(const UE::NNECore::FAttributeMap& AttributeMap, TConstArrayView<EMLTensorDataType> InputTypes, TConstArrayView<NNX::FSymbolicTensorShape> InputShapes)
+	{
+		bool bIsValid = true;
+
+		//This match version 2 of the pad operator (next version is with opset 11)
+		//[REDACTED_PROFILE_URL]/onnx/blob/main/docs/Changelog.md#Pad-2
+		NNX::FAttributeValidator AttributeValidator;
+		AttributeValidator.AddOptional(TEXT("mode"), ENNEAttributeDataType::String);
+		AttributeValidator.AddRequired(TEXT("pads"), ENNEAttributeDataType::Int32Array);
+		AttributeValidator.AddOptional(TEXT("value"), ENNEAttributeDataType::Float);
+		bIsValid &= AttributeValidator.Validate(AttributeMap);
+
+		for (int32 Pad : AttributeMap.GetValue<TArray<int32>>(TEXT("pads")))
+		{
+			if (Pad < 0)
+			{
+				UE_LOG(LogNNX, Warning, TEXT("Pad operator does not support negative padding at the moment."));
+				return false;
+			}
+		}
+		
+		NNX::FInputValidator InputValidator;
+		InputValidator.AddSupportedType(EMLTensorDataType::Float);
+		InputValidator.AddRequired();
+		bIsValid &= InputValidator.Validate(InputTypes);
+
+		return bIsValid;
+	}
+
+	NNX::FMLOperatorHlsl* CreatePadOperator()
+	{
+		return new FPad();
+	}
+
+	bool RegisterPadOperator(NNX::FMLOperatorRegistryHlsl& Registry)
+	{
+		Registry.OpAdd(TEXT("Pad"), CreatePadOperator, ValidatePadOperator);
+		return true;
+	}
+} // UE::NNIRuntimeRDG::Private::Hlsl
diff --git a/Engine/Plugins/Experimental/NNE/Source/NNXRuntimeRDG/Private/Hlsl/NNIRuntimeRDGPad.h b/Engine/Plugins/Experimental/NNE/Source/NNXRuntimeRDG/Private/Hlsl/NNIRuntimeRDGPad.h
new file mode 100644
--- /dev/null
+++ b/Engine/Plugins/Experimental/NNE/Source/NNXRuntimeRDG/Private/Hlsl/NNIRuntimeRDGPad.h
@@ -0,0 +1,11 @@
+// Copyright Epic Games, Inc. All Rights Reserved.
+
+#pragma once
+
+#include "../NNXRuntimeRDG.h"
+#include "../NNXRuntimeHLSLOp.h"
+
+namespace UE::NNIRuntimeRDG::Private::Hlsl
+{
+	bool RegisterPadOperator(NNX::FMLOperatorRegistryHlsl& Registry);
+} // UE::NNIRuntimeRDG::Private::Hlsl
diff --git a/Engine/Plugins/Experimental/NNE/Source/NNXRuntimeRDG/Private/NNXRuntimeHLSL.cpp b/Engine/Plugins/Experimental/NNE/Source/NNXRuntimeRDG/Private/NNXRuntimeHLSL.cpp
--- a/Engine/Plugins/Experimental/NNE/Source/NNXRuntimeRDG/Private/NNXRuntimeHLSL.cpp
+++ b/Engine/Plugins/Experimental/NNE/Source/NNXRuntimeRDG/Private/NNXRuntimeHLSL.cpp
@@ -13,6 +13,7 @@
 #include "Hlsl/NNIRuntimeRDGGemm.h"
 #include "Hlsl/NNIRuntimeRDGMatMul.h"
 #include "Hlsl/NNIRuntimeRDGInstanceNormalization.h"
+#include "Hlsl/NNIRuntimeRDGPad.h"
 #include "Hlsl/NNIRuntimeRDGUpsample.h"
 #include "NNXRuntimeHLSLHelper.h"
 #include "NNXRuntimeRDG.h"
@@ -268,6 +269,7 @@ public:
 		UE::NNIRuntimeRDG::Private::Hlsl::RegisterMatMulOperator(*registry);
 		UE::NNIRuntimeRDG::Private::Hlsl::RegisterInstanceNormalizationOperator(*registry);
 		UE::NNIRuntimeRDG::Private::Hlsl::RegisterUpsampleOperator(*registry);
+		UE::NNIRuntimeRDG::Private::Hlsl::RegisterPadOperator(*registry);
 
 		return true;
 	}
diff --git a/Engine/Plugins/Experimental/NNE/Source/NNXUtils/Private/NNXModelBuilderONNX.cpp b/Engine/Plugins/Experimental/NNE/Source/NNXUtils/Private/NNXModelBuilderONNX.cpp
--- a/Engine/Plugins/Experimental/NNE/Source/NNXUtils/Private/NNXModelBuilderONNX.cpp
+++ b/Engine/Plugins/Experimental/NNE/Source/NNXUtils/Private/NNXModelBuilderONNX.cpp
@@ -334,9 +334,10 @@ NNXUTILS_API bool CreateONNXModelForOperator(bool UseVariadicShapeForModel, cons
 	
 	int64 IrVersion = OnnxIrVersion;
 	int64 OpsetVersion = OnnxOpsetVersion;
-	if (OperatorName == TEXT("Upsample"))
+	if (OperatorName == TEXT("Upsample") || OperatorName == TEXT("Pad"))
 	{
 		// Upsample is deprecated starting opset 10
+		// Pad current implementation is opset 2 (next version is 11)
 		OpsetVersion = 9;
 	}
 	TUniquePtr<IMLModelBuilder> Builder(CreateONNXModelBuilder(IrVersion, OpsetVersion));
