NNE - Adding neural post processing

This CL adds neural post processing to NNE.
A view extension registers a callback for enqueueing post procssing calls.
The added code will downsample the input scene color to the neural network input shape, rune the neural network (currently disabled due to missing weights) and upsample it to overwrite the scene color.
There is a blueprint wrapper class to make this functionality accessible through blueprints. Also a simple sobel filter has been added once with fixed input, once with dynamic shape for testing purposes.

diff --git a/Engine/Plugins/Experimental/NNE/Source/NNEHlslShaders/Internal/NNEHlslShadersNeuralPostProcessingCS.h b/Engine/Plugins/Experimental/NNE/Source/NNEHlslShaders/Internal/NNEHlslShadersNeuralPostProcessingCS.h
new file mode 100644
--- /dev/null
+++ b/Engine/Plugins/Experimental/NNE/Source/NNEHlslShaders/Internal/NNEHlslShadersNeuralPostProcessingCS.h
@@ -0,0 +1,110 @@
+// Copyright Epic Games, Inc. All Rights Reserved.
+
+#pragma once
+
+#include "GlobalShader.h"
+#include "ShaderParameterUtils.h"
+#include "RenderGraphUtils.h"
+
+namespace UE::NNEHlslShaders::Internal
+{
+	enum class ENeuralPostProcessingOverwrite : uint8
+	{
+		No = 0,
+		Yes,
+		MAX
+	};
+
+	enum class ENeuralPostProcessingInterpolate : uint8
+	{
+		No = 0,
+		Yes,
+		MAX
+	};
+
+	class FNeuralPostProcessingConstants
+	{
+	public:
+		static const int32 THREAD_GROUP_SIZE{ 32 };
+	};
+
+	class NNEHLSLSHADERS_API TNeuralPostProcessingReadInputCS : public FGlobalShader
+	{
+		DECLARE_GLOBAL_SHADER(TNeuralPostProcessingReadInputCS);
+		SHADER_USE_PARAMETER_STRUCT(TNeuralPostProcessingReadInputCS, FGlobalShader)
+
+	public:
+
+		BEGIN_SHADER_PARAMETER_STRUCT(FParameters, )
+			SHADER_PARAMETER_RDG_TEXTURE(Texture2D, InputTexture)
+			SHADER_PARAMETER_SAMPLER(SamplerState, InputTextureSampler)
+			SHADER_PARAMETER(int32, InputTextureWidth)
+			SHADER_PARAMETER(int32, InputTextureHeight)
+			SHADER_PARAMETER_RDG_BUFFER_UAV(RWBuffer<float>, AccumulationBuffer)
+			SHADER_PARAMETER(float, Weight)
+		END_SHADER_PARAMETER_STRUCT()
+
+		static void ModifyCompilationEnvironment(const FGlobalShaderPermutationParameters& InParameters, FShaderCompilerEnvironment& OutEnvironment);
+	};
+
+	class NNEHLSLSHADERS_API TNeuralPostProcessingPreStepCS : public FGlobalShader
+	{
+		DECLARE_GLOBAL_SHADER(TNeuralPostProcessingPreStepCS);
+		SHADER_USE_PARAMETER_STRUCT(TNeuralPostProcessingPreStepCS, FGlobalShader)
+
+	public:
+
+		BEGIN_SHADER_PARAMETER_STRUCT(FParameters, )
+			SHADER_PARAMETER_RDG_TEXTURE(Texture2D, InputTexture)
+			SHADER_PARAMETER_SAMPLER(SamplerState, InputTextureSampler)
+			SHADER_PARAMETER(int32, InputTextureWidth)
+			SHADER_PARAMETER(int32, InputTextureHeight)
+			SHADER_PARAMETER(int32, InputBufferWidth)
+			SHADER_PARAMETER(int32, InputBufferHeight)
+			SHADER_PARAMETER_RDG_BUFFER_UAV(RWBuffer<float>, InputBuffer)
+		END_SHADER_PARAMETER_STRUCT()
+
+		static void ModifyCompilationEnvironment(const FGlobalShaderPermutationParameters& InParameters, FShaderCompilerEnvironment& OutEnvironment);
+	};
+
+	class NNEHLSLSHADERS_API TNeuralPostProcessingPostStepCS : public FGlobalShader
+	{
+		DECLARE_GLOBAL_SHADER(TNeuralPostProcessingPostStepCS);
+		SHADER_USE_PARAMETER_STRUCT(TNeuralPostProcessingPostStepCS, FGlobalShader)
+
+		class FNeuralPostProcessingOverwrite : SHADER_PERMUTATION_ENUM_CLASS("OVERWRITE", ENeuralPostProcessingOverwrite);
+		class FNeuralPostProcessingInterpolate : SHADER_PERMUTATION_ENUM_CLASS("INTERPOLATE", ENeuralPostProcessingInterpolate);
+		using FPermutationDomain = TShaderPermutationDomain<FNeuralPostProcessingOverwrite, FNeuralPostProcessingInterpolate>;
+
+	public:
+
+		BEGIN_SHADER_PARAMETER_STRUCT(FParameters, )
+			SHADER_PARAMETER(int32, OutputBufferWidth)
+			SHADER_PARAMETER(int32, OutputBufferHeight)
+			SHADER_PARAMETER_RDG_BUFFER_UAV(RWBuffer<float>, OutputBuffer)
+			SHADER_PARAMETER(int32, InputTextureWidth)
+			SHADER_PARAMETER(int32, InputTextureHeight)
+			SHADER_PARAMETER_RDG_BUFFER_UAV(RWBuffer<float>, AccumulationBuffer)
+			SHADER_PARAMETER(float, Weight)
+		END_SHADER_PARAMETER_STRUCT()
+
+		static void ModifyCompilationEnvironment(const FGlobalShaderPermutationParameters& InParameters, FShaderCompilerEnvironment& OutEnvironment);
+	};
+
+	class NNEHLSLSHADERS_API TNeuralPostProcessingWriteOutputPS : public FGlobalShader
+	{
+		DECLARE_GLOBAL_SHADER(TNeuralPostProcessingWriteOutputPS);
+		SHADER_USE_PARAMETER_STRUCT(TNeuralPostProcessingWriteOutputPS, FGlobalShader)
+
+	public:
+
+		BEGIN_SHADER_PARAMETER_STRUCT(FParameters, )
+			SHADER_PARAMETER_RDG_BUFFER_UAV(RWBuffer<float>, AccumulationBuffer)
+			SHADER_PARAMETER(int32, InputTextureWidth)
+			SHADER_PARAMETER(int32, InputTextureHeight)
+			RENDER_TARGET_BINDING_SLOTS()
+		END_SHADER_PARAMETER_STRUCT()
+
+		static void ModifyCompilationEnvironment(const FGlobalShaderPermutationParameters& InParameters, FShaderCompilerEnvironment& OutEnvironment);
+	};
+} // UE::NNEHlslShaders::Internal
\ No newline at end of file
diff --git a/Engine/Plugins/Experimental/NNE/Source/NNEHlslShaders/Private/NNEHlslShadersNeuralPostProcessingCS.cpp b/Engine/Plugins/Experimental/NNE/Source/NNEHlslShaders/Private/NNEHlslShadersNeuralPostProcessingCS.cpp
new file mode 100644
--- /dev/null
+++ b/Engine/Plugins/Experimental/NNE/Source/NNEHlslShaders/Private/NNEHlslShadersNeuralPostProcessingCS.cpp
@@ -0,0 +1,35 @@
+// Copyright Epic Games, Inc. All Rights Reserved.
+
+#include "NNEHlslShadersNeuralPostProcessingCS.h"
+
+namespace UE::NNEHlslShaders::Internal
+{
+	void TNeuralPostProcessingReadInputCS::ModifyCompilationEnvironment(const FGlobalShaderPermutationParameters& InParameters, FShaderCompilerEnvironment& OutEnvironment)
+	{
+		FGlobalShader::ModifyCompilationEnvironment(InParameters, OutEnvironment);
+		OutEnvironment.SetDefine(TEXT("THREAD_GROUP_SIZE"), FNeuralPostProcessingConstants::THREAD_GROUP_SIZE);
+	}
+
+	void TNeuralPostProcessingPreStepCS::ModifyCompilationEnvironment(const FGlobalShaderPermutationParameters& InParameters, FShaderCompilerEnvironment& OutEnvironment)
+	{
+		FGlobalShader::ModifyCompilationEnvironment(InParameters, OutEnvironment);
+		OutEnvironment.SetDefine(TEXT("THREAD_GROUP_SIZE"), FNeuralPostProcessingConstants::THREAD_GROUP_SIZE);
+	}
+
+	void TNeuralPostProcessingPostStepCS::ModifyCompilationEnvironment(const FGlobalShaderPermutationParameters& InParameters, FShaderCompilerEnvironment& OutEnvironment)
+	{
+		FGlobalShader::ModifyCompilationEnvironment(InParameters, OutEnvironment);
+		OutEnvironment.SetDefine(TEXT("THREAD_GROUP_SIZE"), FNeuralPostProcessingConstants::THREAD_GROUP_SIZE);
+	}
+
+	void TNeuralPostProcessingWriteOutputPS::ModifyCompilationEnvironment(const FGlobalShaderPermutationParameters& InParameters, FShaderCompilerEnvironment& OutEnvironment)
+	{
+		FGlobalShader::ModifyCompilationEnvironment(InParameters, OutEnvironment);
+		OutEnvironment.SetDefine(TEXT("THREAD_GROUP_SIZE"), FNeuralPostProcessingConstants::THREAD_GROUP_SIZE);
+	}
+
+	IMPLEMENT_GLOBAL_SHADER(TNeuralPostProcessingReadInputCS, "/NNE/NNEHlslShadersNeuralPostProcessing.usf", "ReadInput", SF_Compute);
+	IMPLEMENT_GLOBAL_SHADER(TNeuralPostProcessingPreStepCS, "/NNE/NNEHlslShadersNeuralPostProcessing.usf", "PreStep", SF_Compute);
+	IMPLEMENT_GLOBAL_SHADER(TNeuralPostProcessingPostStepCS, "/NNE/NNEHlslShadersNeuralPostProcessing.usf", "PostStep", SF_Compute);
+	IMPLEMENT_GLOBAL_SHADER(TNeuralPostProcessingWriteOutputPS, "/NNE/NNEHlslShadersNeuralPostProcessing.usf", "WriteOutput", SF_Pixel);
+} // UE::NNEHlslShaders::Internal
\ No newline at end of file
diff --git a/Engine/Plugins/Experimental/NNE/Source/NNEHlslShaders/Shaders/NNEHlslShadersNeuralPostProcessing.usf b/Engine/Plugins/Experimental/NNE/Source/NNEHlslShaders/Shaders/NNEHlslShadersNeuralPostProcessing.usf
new file mode 100644
--- /dev/null
+++ b/Engine/Plugins/Experimental/NNE/Source/NNEHlslShaders/Shaders/NNEHlslShadersNeuralPostProcessing.usf
@@ -0,0 +1,153 @@
+// Copyright Epic Games, Inc. All Rights Reserved.
+
+#include "/Engine/Public/Platform.ush"
+#include "/Engine/Private/Common.ush"
+
+int InputTextureWidth;
+int InputTextureHeight;
+Texture2D InputTexture;
+SamplerState InputTextureSampler;
+RWBuffer<float> AccumulationBuffer;
+
+int InputBufferWidth;
+int InputBufferHeight;
+RWBuffer<float> InputBuffer;
+
+int OutputBufferWidth;
+int OutputBufferHeight;
+RWBuffer<float> OutputBuffer;
+
+float Weight;
+
+[numthreads(THREAD_GROUP_SIZE, THREAD_GROUP_SIZE, 1)]
+void ReadInput(in const uint3 DispatchThreadID : SV_DispatchThreadID)
+{
+	if (DispatchThreadID.x >= InputTextureWidth || DispatchThreadID.y >= InputTextureHeight)
+	{
+		return;
+	}
+
+	float4 InputColor = InputTexture.SampleLevel(InputTextureSampler, float2((float(DispatchThreadID.x) + 0.5) / float(InputTextureWidth), (float(DispatchThreadID.y) + 0.5) / float(InputTextureHeight)), 0);
+
+	int Idx = (DispatchThreadID.y * InputTextureWidth + DispatchThreadID.x) * 3;
+	AccumulationBuffer[Idx + 0] = Weight * InputColor.b;
+	AccumulationBuffer[Idx + 1] = Weight * InputColor.g;
+	AccumulationBuffer[Idx + 2] = Weight * InputColor.r;
+}
+
+[numthreads(THREAD_GROUP_SIZE, THREAD_GROUP_SIZE, 1)]
+void PreStep(in const uint3 DispatchThreadID : SV_DispatchThreadID)
+{
+	if (DispatchThreadID.x >= InputBufferWidth || DispatchThreadID.y >= InputBufferHeight)
+	{
+		return;
+	}
+
+	float WidthFactor = (float)InputTextureWidth / (float)InputBufferWidth;
+	int LoX = (int)floor(WidthFactor * (float)DispatchThreadID.x);
+	int HiX = (int)ceil(WidthFactor * (float)(DispatchThreadID.x + 1));
+
+	float HeightFactor = (float)InputTextureHeight / (float)InputBufferHeight;
+	int LoY = (int)floor(HeightFactor * (float)DispatchThreadID.y);
+	int HiY = (int)ceil(HeightFactor * (float)(DispatchThreadID.y + 1));
+
+	float4 Result = float4(0.0, 0.0, 0.0, 0.0);
+	float Div = 0.0;
+	for (int x = LoX; x < HiX; x++)
+	{
+		if (x < InputTextureWidth)
+		{
+			for (int y = LoY; y < HiY; y++)
+			{
+				if (y < InputTextureHeight)
+				{
+					Result += InputTexture.SampleLevel(InputTextureSampler, float2((x + 0.5) / float(InputTextureWidth), (y + 0.5) / float(InputTextureHeight)), 0);
+					Div += 1.0;
+				}
+			}
+		}
+	}
+	Result /= Div;
+
+	int Idx = DispatchThreadID.y * InputBufferWidth + DispatchThreadID.x;
+	int Offset1 = InputBufferWidth * InputBufferHeight;
+	int Offset2 = Offset1 + Offset1;
+	InputBuffer[Idx] = Result.r;
+	InputBuffer[Idx + Offset1] = Result.g;
+	InputBuffer[Idx + Offset2] = Result.b;
+}
+
+[numthreads(THREAD_GROUP_SIZE, THREAD_GROUP_SIZE, 1)]
+void PostStep(in const uint3 DispatchThreadID : SV_DispatchThreadID)
+{
+	if (DispatchThreadID.x >= InputTextureWidth || DispatchThreadID.y >= InputTextureHeight)
+	{
+		return;
+	}
+
+	float WidthFactor = (float)OutputBufferWidth / (float)InputTextureWidth;
+	float X = WidthFactor * ((float)DispatchThreadID.x + 0.5);
+
+	float HeightFactor = (float)OutputBufferHeight / (float)InputTextureHeight;
+	float Y = HeightFactor * ((float)DispatchThreadID.y + 0.5);
+
+	int Offset1 = OutputBufferWidth * OutputBufferHeight;
+	int Offset2 = Offset1 + Offset1;
+
+	float3 Result;
+#if INTERPOLATE == 0
+	int Idx = (int)Y * OutputBufferWidth + (int)X;
+	Result.r = OutputBuffer[Idx];
+	Result.g = OutputBuffer[Idx + Offset1];
+	Result.b = OutputBuffer[Idx + Offset2];
+#else
+	int LoX = (int)clamp(floor(X - 0.5), 0.0, float(OutputBufferWidth - 1));
+	int HiX = (int)clamp(floor(X + 0.5), 0.0, float(OutputBufferWidth - 1));
+
+	int LoY = (int)clamp(floor(Y - 0.5), 0.0, float(OutputBufferHeight - 1));
+	int HiY = (int)clamp(floor(Y + 0.5), 0.0, float(OutputBufferHeight - 1));
+
+	int Idx;
+	
+	Idx = LoY * OutputBufferWidth + LoX;
+	float3 A = float3(OutputBuffer[Idx], OutputBuffer[Idx + Offset1], OutputBuffer[Idx + Offset2]);
+
+	Idx = LoY * OutputBufferWidth + HiX;
+	float3 B = float3(OutputBuffer[Idx], OutputBuffer[Idx + Offset1], OutputBuffer[Idx + Offset2]);
+
+	Idx = HiY * OutputBufferWidth + LoX;
+	float3 C = float3(OutputBuffer[Idx], OutputBuffer[Idx + Offset1], OutputBuffer[Idx + Offset2]);
+
+	Idx = HiY * OutputBufferWidth + HiX;
+	float3 D = float3(OutputBuffer[Idx], OutputBuffer[Idx + Offset1], OutputBuffer[Idx + Offset2]);
+
+	float Alpha = clamp((X - 0.5) - ((float)LoX), 0.0, 1.0);
+	float Beta = clamp((Y - 0.5) - ((float)LoY), 0.0, 1.0);
+	
+	Result = (1.0 - Beta) * ((1.0 - Alpha) * A + Alpha * B) + Beta * ((1.0 - Alpha) * C + Alpha * D);
+#endif
+
+	int ResultIdx = (DispatchThreadID.y * InputTextureWidth + DispatchThreadID.x) * 3;
+#if OVERWRITE == 0
+	AccumulationBuffer[ResultIdx + 0] += Weight * Result.r;
+	AccumulationBuffer[ResultIdx + 1] += Weight * Result.g;
+	AccumulationBuffer[ResultIdx + 2] += Weight * Result.b;
+#else
+	AccumulationBuffer[ResultIdx + 0] = Weight * Result.r;
+	AccumulationBuffer[ResultIdx + 1] = Weight * Result.g;
+	AccumulationBuffer[ResultIdx + 2] = Weight * Result.b;
+#endif
+}
+
+float4 WriteOutput(float4 Position : SV_POSITION) : SV_Target0
+{
+	int Idx = ((int)Position.y * InputTextureWidth + (int)Position.x) * 3;
+
+	float4 Result;
+	Result.r = AccumulationBuffer[Idx + 0];
+	Result.g = AccumulationBuffer[Idx + 1];
+	Result.b = AccumulationBuffer[Idx + 2];
+	Result.a = 1.0;
+
+	return Result;
+}
\ No newline at end of file
diff --git a/Engine/Plugins/Experimental/NNE/Source/NNXRuntimeRDG/NNXRuntimeRDG.Build.cs b/Engine/Plugins/Experimental/NNE/Source/NNXRuntimeRDG/NNXRuntimeRDG.Build.cs
--- a/Engine/Plugins/Experimental/NNE/Source/NNXRuntimeRDG/NNXRuntimeRDG.Build.cs
+++ b/Engine/Plugins/Experimental/NNE/Source/NNXRuntimeRDG/NNXRuntimeRDG.Build.cs
@@ -1,5 +1,6 @@
 // Copyright Epic Games, Inc. All Rights Reserved.
 
+using System.IO;
 using UnrealBuildTool;
 
 public class NNXRuntimeRDG : ModuleRules
@@ -7,7 +8,9 @@ public class NNXRuntimeRDG : ModuleRules
 	public NNXRuntimeRDG(ReadOnlyTargetRules Target) : base(Target)
 	{
 		PCHUsage = PCHUsageMode.UseExplicitOrSharedPCHs;
-		
+
+		PrivateIncludePaths.AddRange(new string[] { Path.Combine(EngineDirectory, "Source/Runtime/Renderer/Private") });
+
 		PublicDependencyModuleNames.AddRange(new string[] 
 		{ 
 			"Core", 
diff --git a/Engine/Plugins/Experimental/NNE/Source/NNXRuntimeRDG/Private/Hlsl/NNIRuntimeRDGConv.cpp b/Engine/Plugins/Experimental/NNE/Source/NNXRuntimeRDG/Private/Hlsl/NNIRuntimeRDGConv.cpp
--- a/Engine/Plugins/Experimental/NNE/Source/NNXRuntimeRDG/Private/Hlsl/NNIRuntimeRDGConv.cpp
+++ b/Engine/Plugins/Experimental/NNE/Source/NNXRuntimeRDG/Private/Hlsl/NNIRuntimeRDGConv.cpp
@@ -177,7 +177,7 @@ namespace UE::NNIRuntimeRDG::Private::Hlsl
 		AttributeValidator.AddOptional(TEXT("auto_pad"), ENNEAttributeDataType::String);
 		AttributeValidator.AddOptional(TEXT("dilations"), ENNEAttributeDataType::Int32Array);
 		AttributeValidator.AddOptional(TEXT("group"), ENNEAttributeDataType::Int32);
-		//AttributeValidator.AddOptional(TEXT("kernel_shape"), ENNEAttributeDataType::Int32Array);
+		AttributeValidator.AddOptional(TEXT("kernel_shape"), ENNEAttributeDataType::Int32Array); // idea: cross check input weight shape with this attribute if present
 		AttributeValidator.AddOptional(TEXT("pads"), ENNEAttributeDataType::Int32Array);
 		AttributeValidator.AddOptional(TEXT("strides"), ENNEAttributeDataType::Int32Array);
 
diff --git a/Engine/Plugins/Experimental/NNE/Source/NNXRuntimeRDG/Private/NNERuntimeRDGNeuralPostProcessing.cpp b/Engine/Plugins/Experimental/NNE/Source/NNXRuntimeRDG/Private/NNERuntimeRDGNeuralPostProcessing.cpp
new file mode 100644
--- /dev/null
+++ b/Engine/Plugins/Experimental/NNE/Source/NNXRuntimeRDG/Private/NNERuntimeRDGNeuralPostProcessing.cpp
@@ -0,0 +1,388 @@
+// Copyright Epic Games, Inc. All Rights Reserved.
+
+#include "NNERuntimeRDGNeuralPostProcessing.h"
+#include "NNXCore.h"
+#include "PostProcess/PostProcessMaterial.h"
+#include "PostProcess/SceneRenderTargets.h"
+#include "ScreenPass.h"
+#include "PostProcess/PostProcessing.h"
+#include "PixelShaderUtils.h"
+
+#include "NNEHlslShadersNeuralPostProcessingCS.h"
+
+DECLARE_GPU_STAT_NAMED(FNNENeuralPostProcessingReadInput, TEXT("NNE.NeuralPostProcessing.ReadInput"));
+DECLARE_GPU_STAT_NAMED(FNNENeuralPostProcessingPreStep, TEXT("NNE.NeuralPostProcessing.PreStep"));
+DECLARE_GPU_STAT_NAMED(FNNENeuralPostProcessingPostStep, TEXT("NNE.NeuralPostProcessing.PostStep"));
+DECLARE_GPU_STAT_NAMED(FNNENeuralPostProcessingWriteOutput, TEXT("NNE.NeuralPostProcessing.WriteOutput"));
+
+FNNENeuralPostProcessing::FNNENeuralPostProcessing(const FAutoRegister& AutoRegister) : FSceneViewExtensionBase(AutoRegister) 
+{
+	NumEnabled = 0;
+	LastId = 0;
+}
+
+int32 FNNENeuralPostProcessing::Add(FString RuntimeName, UNNEModelData* ModelData)
+{
+	FScopeLock Lock(&CriticalSection);
+
+	// Create the model
+	NNX::IRuntime* Runtime = NNX::GetRuntime(RuntimeName);
+	if (!Runtime)
+	{
+		UE_LOG(LogNNX, Error, TEXT("FNNENeuralPostProcessing: No runtime '%s' found. Valid runtimes are: "), *RuntimeName);
+		TArray<NNX::IRuntime*> Runtimes = NNX::GetAllRuntimes();
+		for (int32 i = 0; i < Runtimes.Num(); i++)
+		{
+			UE_LOG(LogNNX, Error, TEXT("- %s"), *Runtimes[i]->GetRuntimeName());
+		}
+		return -1;
+	}
+
+	if (!ModelData)
+	{
+		UE_LOG(LogNNX, Error, TEXT("FNNENeuralPostProcessing: Valid model data required to load the model"));
+		return -1;
+	}
+
+	TConstArrayView<uint8> Data = ModelData->GetModelData(RuntimeName);
+	if (Data.Num() < 1)
+	{
+		UE_LOG(LogNNX, Error, TEXT("FNNENeuralPostProcessing: No model data for %s found"), *RuntimeName);
+		return -1;
+	}
+
+	TSharedPtr<NNX::FMLInferenceModel> Model = TSharedPtr<NNX::FMLInferenceModel>(Runtime->CreateModel(Data).Release());
+	if (!Model.IsValid())
+	{
+		UE_LOG(LogNNX, Error, TEXT("FNNENeuralPostProcessing: Could not create model using %s"), *RuntimeName);
+		return -1;
+	}
+
+	// Create a new id
+	LastId = (LastId + 1) > 0 ? (LastId + 1) : 1;
+
+	// Add the model to the map
+	Models.Add(LastId, Model);
+	SetWeight(LastId, 0.0);
+
+	return LastId;
+}
+
+bool FNNENeuralPostProcessing::Remove(int32 ModelId)
+{
+	FScopeLock Lock(&CriticalSection);
+
+	bool bResult = Models.Find(ModelId) != nullptr;
+
+	Disable(ModelId);
+	Models.Remove(ModelId);
+	Weights.Remove(ModelId);
+
+	return bResult;
+}
+
+bool FNNENeuralPostProcessing::SetWeight(int32 ModelId, float Weight)
+{
+	FScopeLock Lock(&CriticalSection);
+
+	if (Models.Find(ModelId) != nullptr)
+	{
+		Weights.Add(ModelId, Weight);
+		return true;
+	}
+	return false;
+}
+
+void FNNENeuralPostProcessing::Enable(int32 ModelId)
+{
+	FScopeLock Lock(&CriticalSection);
+
+	if (Enabled.Find(ModelId) == nullptr)
+	{
+		Enabled.Add(ModelId);
+		NumEnabled++;
+	}
+}
+
+void FNNENeuralPostProcessing::Disable(int32 ModelId)
+{
+	FScopeLock Lock(&CriticalSection);
+
+	if (Enabled.Find(ModelId) != nullptr)
+	{
+		Enabled.Remove(ModelId);
+		NumEnabled--;
+	}
+}
+
+void FNNENeuralPostProcessing::PrePostProcessPass_RenderThread(FRDGBuilder& GraphBuilder, const FSceneView& View, const FPostProcessingInputs& Inputs)
+{
+	using namespace UE::NNEHlslShaders::Internal;
+	const float WeightEpsilon = 1.0/65536.0;
+
+	check(IsInRenderingThread());
+	check(View.bIsViewInfo);
+
+	float InputWeight = 0.0;
+	{
+		FScopeLock Lock(&CriticalSection);
+
+		if (NumEnabled < 1)
+		{
+			return;
+		}
+
+		float WeightSum = 0.0;
+		for (TPair<int32, float> Pair : Weights)
+		{
+			if (Enabled.Find(Pair.Key) != nullptr)
+			{
+				if (Pair.Value >= 0.0)
+				{
+					WeightSum += Pair.Value;
+				}
+				else
+				{
+					Weights.Add(Pair.Key, 0.0);
+				}
+			}
+		}
+
+		if (WeightSum > 1.0)
+		{
+			for (TPair<int32, float> Pair : Weights)
+			{
+				if (Enabled.Find(Pair.Key) != nullptr)
+				{
+					Weights.Add(Pair.Key, Pair.Value / WeightSum);
+				}
+			}
+		}
+		else
+		{
+			InputWeight = 1.0 - WeightSum;
+		}
+	}
+
+	const FIntRect Viewport = static_cast<const FViewInfo&>(View).ViewRect; 
+	FScreenPassTexture SceneColor((*Inputs.SceneTextures)->SceneColorTexture, Viewport);
+	FGlobalShaderMap* GlobalShaderMap = GetGlobalShaderMap(GMaxRHIFeatureLevel);
+	FIntPoint TextureSize = (*Inputs.SceneTextures)->SceneColorTexture->Desc.Extent;
+
+	// Read the input into an accumulation buffer 
+	FRDGBufferDesc AccumulationBufferDesc;
+	AccumulationBufferDesc.Usage = EBufferUsageFlags::UnorderedAccess | EBufferUsageFlags::ShaderResource | EBufferUsageFlags::StructuredBuffer;
+	AccumulationBufferDesc.BytesPerElement = sizeof(float);
+	AccumulationBufferDesc.NumElements = TextureSize.X * TextureSize.Y * 3;
+	FRDGBufferRef AccumulationBuffer = GraphBuilder.CreateBuffer(AccumulationBufferDesc, *FString("NNENeuralPostProcessing::AccumulationBuffer"));
+	FRDGBufferUAVRef AccumulationBufferUAV = GraphBuilder.CreateUAV(AccumulationBuffer);
+
+	bool bOverwrite = true;
+	if (InputWeight > WeightEpsilon)
+	{
+		TNeuralPostProcessingReadInputCS::FParameters* ReadInputParameters = GraphBuilder.AllocParameters<TNeuralPostProcessingReadInputCS::FParameters>();
+		ReadInputParameters->InputTexture = (*Inputs.SceneTextures)->SceneColorTexture;
+		ReadInputParameters->InputTextureSampler = TStaticSamplerState<SF_Point, AM_Clamp, AM_Clamp, AM_Clamp>::GetRHI();
+		ReadInputParameters->InputTextureWidth = TextureSize.X;
+		ReadInputParameters->InputTextureHeight = TextureSize.Y;
+		ReadInputParameters->AccumulationBuffer = AccumulationBufferUAV;
+		ReadInputParameters->Weight = InputWeight;
+
+		FIntVector ReadInputThreadGroupCount = FIntVector(FMath::DivideAndRoundUp(TextureSize.X, FNeuralPostProcessingConstants::THREAD_GROUP_SIZE), FMath::DivideAndRoundUp(TextureSize.Y, FNeuralPostProcessingConstants::THREAD_GROUP_SIZE), 1);
+		TShaderMapRef<TNeuralPostProcessingReadInputCS> ReadInputShader(GlobalShaderMap);
+
+		RDG_EVENT_SCOPE(GraphBuilder, "NNE.NeuralPostProcessing.ReadInput");
+		RDG_GPU_STAT_SCOPE(GraphBuilder, FNNENeuralPostProcessingReadInput);
+
+		FComputeShaderUtils::AddPass(
+			GraphBuilder,
+			RDG_EVENT_NAME("NNE.NeuralPostProcessing.ReadInput"),
+			ERDGPassFlags::Compute | ERDGPassFlags::NeverCull,
+			ReadInputShader,
+			ReadInputParameters,
+			ReadInputThreadGroupCount);
+
+		bOverwrite = false;
+	}
+
+	// Process each network
+	{
+		FScopeLock Lock(&CriticalSection);
+		
+		for (TPair<int32, TSharedPtr<NNX::FMLInferenceModel, ESPMode::ThreadSafe>> Pair : Models)
+		{
+			if (Enabled.Find(Pair.Key) != nullptr && Weights.Find(Pair.Key) != nullptr)
+			{
+				float OutputWeight = *Weights.Find(Pair.Key);
+				if (OutputWeight > WeightEpsilon)
+				{
+					NNX::FSymbolicTensorShape InputShape = Pair.Value->GetInputTensorDescs()[0].GetShape();
+
+					checkf(InputShape.Rank() == 4, TEXT("Neural Post Processing requires models with input shape 1 x 3 x height x width!"))
+					checkf(InputShape.Data[0] == 1, TEXT("Neural Post Processing requires models with input shape 1 x 3 x height x width!"))
+					checkf(InputShape.Data[1] == 3, TEXT("Neural Post Processing requires models with input shape 1 x 3 x height x width!"))
+
+					int32 NeuralNetworkInputWidth = InputShape.Data[3] < 0 ? TextureSize.X : InputShape.Data[3];
+					int32 NeuralNetworkInputHeight = InputShape.Data[2] < 0 ? TextureSize.Y : InputShape.Data[2];
+
+					FRDGBufferDesc InputBufferDesc;
+					InputBufferDesc.Usage = EBufferUsageFlags::UnorderedAccess | EBufferUsageFlags::ShaderResource | EBufferUsageFlags::StructuredBuffer;
+					InputBufferDesc.BytesPerElement = sizeof(float);
+					InputBufferDesc.NumElements = NeuralNetworkInputWidth * NeuralNetworkInputHeight * 3;
+
+					FRDGBufferRef InputBuffer = GraphBuilder.CreateBuffer(InputBufferDesc, *(FString("NNENeuralPostProcessing::NeuralNetowrkInput_") + FString::FromInt(Pair.Key)));
+					FRDGBufferUAVRef InputBufferUAV = GraphBuilder.CreateUAV(InputBuffer);
+
+					TNeuralPostProcessingPreStepCS::FParameters* PreStepParameters = GraphBuilder.AllocParameters<TNeuralPostProcessingPreStepCS::FParameters>();
+					PreStepParameters->InputTexture = (*Inputs.SceneTextures)->SceneColorTexture;
+					PreStepParameters->InputTextureSampler = TStaticSamplerState<SF_Point, AM_Clamp, AM_Clamp, AM_Clamp>::GetRHI();
+					PreStepParameters->InputTextureWidth = TextureSize.X;
+					PreStepParameters->InputTextureHeight = TextureSize.Y;
+					PreStepParameters->InputBufferWidth = NeuralNetworkInputWidth;
+					PreStepParameters->InputBufferHeight = NeuralNetworkInputHeight;
+					PreStepParameters->InputBuffer = InputBufferUAV;
+
+					FIntVector PreStepThreadGroupCount = FIntVector(FMath::DivideAndRoundUp(NeuralNetworkInputWidth, FNeuralPostProcessingConstants::THREAD_GROUP_SIZE), FMath::DivideAndRoundUp(NeuralNetworkInputHeight, FNeuralPostProcessingConstants::THREAD_GROUP_SIZE), 1);
+					TShaderMapRef<TNeuralPostProcessingPreStepCS> PreStepShader(GlobalShaderMap);
+
+					RDG_EVENT_SCOPE(GraphBuilder, "NNE.NeuralPostProcessing.PreStep");
+					RDG_GPU_STAT_SCOPE(GraphBuilder, FNNENeuralPostProcessingPreStep);
+
+					FComputeShaderUtils::AddPass(
+						GraphBuilder,
+						RDG_EVENT_NAME("NNE.NeuralPostProcessing.PreStep"),
+						ERDGPassFlags::Compute | ERDGPassFlags::NeverCull,
+						PreStepShader,
+						PreStepParameters,
+						PreStepThreadGroupCount);
+
+					TArray<NNX::FTensorShape> InputShapes;
+					InputShapes.Add(NNX::FTensorShape());
+					InputShapes[0].Data.SetNumUninitialized(4);
+					InputShapes[0].Data[0] = 1;
+					InputShapes[0].Data[1] = 3;
+					InputShapes[0].Data[2] = NeuralNetworkInputHeight;
+					InputShapes[0].Data[3] = NeuralNetworkInputWidth;
+					Pair.Value->SetInputTensorShapes(InputShapes);
+					NNX::FTensorShape OutputShape = Pair.Value->GetOutputTensorShapes()[0];
+
+					checkf(OutputShape.Rank() == 4, TEXT("Neural Post Processing requires models with output shape 1 x 3 x height x width!"))
+					checkf(OutputShape.Data[0] == 1, TEXT("Neural Post Processing requires models with output shape 1 x 3 x height x width!"))
+					checkf(OutputShape.Data[1] == 3, TEXT("Neural Post Processing requires models with output shape 1 x 3 x height x width!"))
+					checkf(OutputShape.Data[2] > 0, TEXT("Neural Post Processing requires models with output height > 0!"))
+					checkf(OutputShape.Data[3] > 0, TEXT("Neural Post Processing requires models with output width > 0!"))
+
+					int32 NeuralNetworkOutputWidth = OutputShape.Data[3];
+					int32 NeuralNetworkOutputHeight = OutputShape.Data[2];
+
+					FRDGBufferDesc OutputBufferDesc;
+					OutputBufferDesc.Usage = EBufferUsageFlags::UnorderedAccess | EBufferUsageFlags::ShaderResource | EBufferUsageFlags::StructuredBuffer;
+					OutputBufferDesc.BytesPerElement = sizeof(float);
+					OutputBufferDesc.NumElements = NeuralNetworkOutputWidth * NeuralNetworkOutputHeight * 3;
+
+					FRDGBufferRef OutputBuffer = GraphBuilder.CreateBuffer(OutputBufferDesc, *(FString("NNENeuralPostProcessing::NeuralNetowrkOutput_") + FString::FromInt(Pair.Key)));
+					FRDGBufferUAVRef OutputBufferUAV = GraphBuilder.CreateUAV(OutputBuffer);
+
+					TArray<NNX::FMLTensorBinding> InputBindings;
+					InputBindings.Add(NNX::FMLTensorBinding::FromRDG(InputBuffer, InputBufferDesc.NumElements * InputBufferDesc.BytesPerElement, 0));
+					TArray<NNX::FMLTensorBinding> OutputBindings;
+					OutputBindings.Add(NNX::FMLTensorBinding::FromRDG(OutputBuffer, OutputBufferDesc.NumElements * OutputBufferDesc.BytesPerElement, 0));
+					// Pair.Value->EnqueueRDG(GraphBuilder, InputBindings, OutputBindings);
+
+					TNeuralPostProcessingPostStepCS::FParameters* PostStepParameters = GraphBuilder.AllocParameters<TNeuralPostProcessingPostStepCS::FParameters>();
+					PostStepParameters->OutputBufferWidth = NeuralNetworkOutputWidth;
+					PostStepParameters->OutputBufferHeight = NeuralNetworkOutputHeight;
+					PostStepParameters->OutputBuffer = InputBufferUAV;// OutputBufferUAV;
+					PostStepParameters->InputTextureWidth = TextureSize.X;
+					PostStepParameters->InputTextureHeight = TextureSize.Y;
+					PostStepParameters->AccumulationBuffer = AccumulationBufferUAV;
+					PostStepParameters->Weight = OutputWeight;
+
+					TNeuralPostProcessingPostStepCS::FPermutationDomain PermutationVector;
+					PermutationVector.Set<TNeuralPostProcessingPostStepCS::FNeuralPostProcessingOverwrite>(bOverwrite ? ENeuralPostProcessingOverwrite::Yes : ENeuralPostProcessingOverwrite::No);
+					PermutationVector.Set<TNeuralPostProcessingPostStepCS::FNeuralPostProcessingInterpolate>((NeuralNetworkOutputWidth == TextureSize.X && NeuralNetworkOutputHeight == TextureSize.Y) ? ENeuralPostProcessingInterpolate::No : ENeuralPostProcessingInterpolate::Yes);
+
+					FIntVector PostStepThreadGroupCount = FIntVector(FMath::DivideAndRoundUp(TextureSize.X, FNeuralPostProcessingConstants::THREAD_GROUP_SIZE), FMath::DivideAndRoundUp(TextureSize.Y, FNeuralPostProcessingConstants::THREAD_GROUP_SIZE), 1);
+					TShaderMapRef<TNeuralPostProcessingPostStepCS> PostStepShader(GlobalShaderMap, PermutationVector);
+
+					RDG_EVENT_SCOPE(GraphBuilder, "NNE.NeuralPostProcessing.PostStep");
+					RDG_GPU_STAT_SCOPE(GraphBuilder, FNNENeuralPostProcessingPostStep);
+
+					FComputeShaderUtils::AddPass(
+						GraphBuilder,
+						RDG_EVENT_NAME("NNE.NeuralPostProcessing.PostStep"),
+						ERDGPassFlags::Compute | ERDGPassFlags::NeverCull,
+						PostStepShader,
+						PostStepParameters,
+						PostStepThreadGroupCount);
+
+					bOverwrite = false;
+				}
+			}
+		}
+	}
+
+	// Write the result
+	TNeuralPostProcessingWriteOutputPS::FParameters* WriteOutputParameters = GraphBuilder.AllocParameters<TNeuralPostProcessingWriteOutputPS::FParameters>();
+	WriteOutputParameters->AccumulationBuffer = AccumulationBufferUAV;
+	WriteOutputParameters->InputTextureWidth = TextureSize.X;
+	WriteOutputParameters->InputTextureHeight = TextureSize.Y;
+	WriteOutputParameters->RenderTargets[0] = FRenderTargetBinding(SceneColor.Texture, ERenderTargetLoadAction::ENoAction);
+
+	RDG_EVENT_SCOPE(GraphBuilder, "NNE.NeuralPostProcessing.WriteOutput");
+	RDG_GPU_STAT_SCOPE(GraphBuilder, FNNENeuralPostProcessingWriteOutput);
+
+	TShaderMapRef<TNeuralPostProcessingWriteOutputPS> WriteOutputShader(GlobalShaderMap);
+	FPixelShaderUtils::AddFullscreenPass(
+		GraphBuilder,
+		GlobalShaderMap,
+		RDG_EVENT_NAME("NNE.NeuralPostProcessing.WriteOutput"),
+		WriteOutputShader,
+		WriteOutputParameters,
+		Viewport);
+}
+
+int32 UNNENeuralPostProcessing::Add(FString RuntimeName, UNNEModelData* ModelData)
+{
+	if (!NeuralPostProcessing.IsValid())
+	{
+		NeuralPostProcessing = FSceneViewExtensions::NewExtension<FNNENeuralPostProcessing>();
+	}
+	return NeuralPostProcessing->Add(RuntimeName, ModelData);
+}
+
+bool UNNENeuralPostProcessing::Remove(int32 ModelId)
+{
+	if (!NeuralPostProcessing.IsValid())
+	{
+		NeuralPostProcessing = FSceneViewExtensions::NewExtension<FNNENeuralPostProcessing>();
+	}
+	return NeuralPostProcessing->Remove(ModelId);
+}
+
+bool UNNENeuralPostProcessing::SetWeight(int32 ModelId, float Weight)
+{
+	if (!NeuralPostProcessing.IsValid())
+	{
+		NeuralPostProcessing = FSceneViewExtensions::NewExtension<FNNENeuralPostProcessing>();
+	}
+	return NeuralPostProcessing->SetWeight(ModelId, Weight);
+}
+
+void UNNENeuralPostProcessing::Enable(int32 ModelId)
+{
+	if (!NeuralPostProcessing.IsValid())
+	{
+		NeuralPostProcessing = FSceneViewExtensions::NewExtension<FNNENeuralPostProcessing>();
+	}
+	NeuralPostProcessing->Enable(ModelId);
+}
+
+void UNNENeuralPostProcessing::Disable(int32 ModelId)
+{
+	if (!NeuralPostProcessing.IsValid())
+	{
+		NeuralPostProcessing = FSceneViewExtensions::NewExtension<FNNENeuralPostProcessing>();
+	}
+	NeuralPostProcessing->Disable(ModelId);
+}
\ No newline at end of file
diff --git a/Engine/Plugins/Experimental/NNE/Source/NNXRuntimeRDG/Public/NNERuntimeRDGNeuralPostProcessing.h b/Engine/Plugins/Experimental/NNE/Source/NNXRuntimeRDG/Public/NNERuntimeRDGNeuralPostProcessing.h
new file mode 100644
--- /dev/null
+++ b/Engine/Plugins/Experimental/NNE/Source/NNXRuntimeRDG/Public/NNERuntimeRDGNeuralPostProcessing.h
@@ -0,0 +1,73 @@
+// Copyright Epic Games, Inc. All Rights Reserved.
+
+#pragma once
+
+#include "CoreMinimal.h"
+#include "SceneViewExtension.h"
+
+#include "NNECoreModelData.h"
+#include "NNXInferenceModel.h"
+
+#include "NNERuntimeRDGNeuralPostProcessing.generated.h"
+
+class NNXRUNTIMERDG_API FNNENeuralPostProcessing : public FSceneViewExtensionBase
+{
+public:
+
+	FNNENeuralPostProcessing(const FAutoRegister& AutoRegister);
+
+public:
+
+	int32 Add(FString RuntimeName, UNNEModelData* ModelData);
+	bool Remove(int32 ModelId);
+	bool SetWeight(int32 ModelId, float Weight);
+	void Enable(int32 ModelId);
+	void Disable(int32 ModelId);
+
+public:
+
+	virtual void SetupViewFamily(FSceneViewFamily& InViewFamily) override {}
+	virtual void SetupView(FSceneViewFamily& InViewFamily, FSceneView& InView) override {}
+	virtual void BeginRenderViewFamily(FSceneViewFamily& InViewFamily) override {}
+	virtual void PreRenderViewFamily_RenderThread(FRHICommandListImmediate& RHICmdList, FSceneViewFamily& InViewFamily) override {}
+	virtual void PreRenderView_RenderThread(FRHICommandListImmediate& RHICmdList, FSceneView& InView) override {}
+	virtual void PrePostProcessPass_RenderThread(FRDGBuilder& GraphBuilder, const FSceneView& View, const FPostProcessingInputs& Inputs) override;
+
+private:
+
+	FCriticalSection CriticalSection;
+	int32 NumEnabled;
+	int32 LastId;
+	TMap<int32, TSharedPtr<NNX::FMLInferenceModel, ESPMode::ThreadSafe>> Models;
+	TMap<int32, float> Weights;
+	TSet<int32> Enabled;
+
+};
+
+UCLASS(BlueprintType, Category = "NNE - Neural Network Engine")
+class NNXRUNTIMERDG_API UNNENeuralPostProcessing : public UObject
+{
+	GENERATED_BODY()
+
+public:
+
+	UFUNCTION(BlueprintCallable, Category = "NNE - Neural Network Engine")
+	int32 Add(FString RuntimeName, UNNEModelData* ModelData);
+
+	UFUNCTION(BlueprintCallable, Category = "NNE - Neural Network Engine")
+	bool Remove(int32 ModelId);
+
+	UFUNCTION(BlueprintCallable, Category = "NNE - Neural Network Engine")
+	bool SetWeight(int32 ModelId, float Weight);
+
+	UFUNCTION(BlueprintCallable, Category = "NNE - Neural Network Engine")
+	void Enable(int32 ModelId);
+
+	UFUNCTION(BlueprintCallable, Category = "NNE - Neural Network Engine")
+	void Disable(int32 ModelId);
+
+private:
+
+	TSharedPtr<FNNENeuralPostProcessing> NeuralPostProcessing;
+
+};
\ No newline at end of file
