NNE - Buffer creation fix to work on consoles

diff --git a/Engine/Plugins/Experimental/NNE/Source/NNXRuntimeRDG/Private/NNERuntimeRDGNeuralPostProcessing.cpp b/Engine/Plugins/Experimental/NNE/Source/NNXRuntimeRDG/Private/NNERuntimeRDGNeuralPostProcessing.cpp
--- a/Engine/Plugins/Experimental/NNE/Source/NNXRuntimeRDG/Private/NNERuntimeRDGNeuralPostProcessing.cpp
+++ b/Engine/Plugins/Experimental/NNE/Source/NNXRuntimeRDG/Private/NNERuntimeRDGNeuralPostProcessing.cpp
@@ -170,12 +170,9 @@ void FNNENeuralPostProcessing::PrePostProcessPass_RenderThread(FRDGBuilder& Grap
 	FIntPoint TextureSize = (*Inputs.SceneTextures)->SceneColorTexture->Desc.Extent;
 
 	// Read the input into an accumulation buffer 
-	FRDGBufferDesc AccumulationBufferDesc;
-	AccumulationBufferDesc.Usage = EBufferUsageFlags::UnorderedAccess | EBufferUsageFlags::ShaderResource | EBufferUsageFlags::StructuredBuffer;
-	AccumulationBufferDesc.BytesPerElement = sizeof(float);
-	AccumulationBufferDesc.NumElements = TextureSize.X * TextureSize.Y * 3;
+	FRDGBufferDesc AccumulationBufferDesc = FRDGBufferDesc::CreateBufferDesc(sizeof(float), TextureSize.X * TextureSize.Y * 3);
 	FRDGBufferRef AccumulationBuffer = GraphBuilder.CreateBuffer(AccumulationBufferDesc, *FString("NNENeuralPostProcessing::AccumulationBuffer"));
-	FRDGBufferUAVRef AccumulationBufferUAV = GraphBuilder.CreateUAV(AccumulationBuffer);
+	FRDGBufferUAVRef AccumulationBufferUAV = GraphBuilder.CreateUAV(FRDGBufferUAVDesc(AccumulationBuffer, PF_R32_FLOAT));
 
 	bool bOverwrite = true;
 	if (InputWeight > WeightEpsilon)
@@ -225,13 +222,9 @@ void FNNENeuralPostProcessing::PrePostProcessPass_RenderThread(FRDGBuilder& Grap
 					int32 NeuralNetworkInputWidth = InputShape.Data[3] < 0 ? TextureSize.X : InputShape.Data[3];
 					int32 NeuralNetworkInputHeight = InputShape.Data[2] < 0 ? TextureSize.Y : InputShape.Data[2];
 
-					FRDGBufferDesc InputBufferDesc;
-					InputBufferDesc.Usage = EBufferUsageFlags::UnorderedAccess | EBufferUsageFlags::ShaderResource | EBufferUsageFlags::StructuredBuffer;
-					InputBufferDesc.BytesPerElement = sizeof(float);
-					InputBufferDesc.NumElements = NeuralNetworkInputWidth * NeuralNetworkInputHeight * 3;
-
+					FRDGBufferDesc InputBufferDesc = FRDGBufferDesc::CreateBufferDesc(sizeof(float), NeuralNetworkInputWidth * NeuralNetworkInputHeight * 3);
 					FRDGBufferRef InputBuffer = GraphBuilder.CreateBuffer(InputBufferDesc, *(FString("NNENeuralPostProcessing::NeuralNetowrkInput_") + FString::FromInt(Pair.Key)));
-					FRDGBufferUAVRef InputBufferUAV = GraphBuilder.CreateUAV(InputBuffer);
+					FRDGBufferUAVRef InputBufferUAV = GraphBuilder.CreateUAV(FRDGBufferUAVDesc(InputBuffer, PF_R32_FLOAT));
 
 					TNeuralPostProcessingPreStepCS::FParameters* PreStepParameters = GraphBuilder.AllocParameters<TNeuralPostProcessingPreStepCS::FParameters>();
 					PreStepParameters->InputTexture = (*Inputs.SceneTextures)->SceneColorTexture;
@@ -257,12 +250,8 @@ void FNNENeuralPostProcessing::PrePostProcessPass_RenderThread(FRDGBuilder& Grap
 						PreStepThreadGroupCount);
 
 					TArray<NNX::FTensorShape> InputShapes;
-					InputShapes.Add(NNX::FTensorShape());
-					InputShapes[0].Data.SetNumUninitialized(4);
-					InputShapes[0].Data[0] = 1;
-					InputShapes[0].Data[1] = 3;
-					InputShapes[0].Data[2] = NeuralNetworkInputHeight;
-					InputShapes[0].Data[3] = NeuralNetworkInputWidth;
+					InputShapes.Add(NNX::FTensorShape::Make({ 1, 3, (uint32)NeuralNetworkInputHeight, (uint32)NeuralNetworkInputWidth }));
+
 					Pair.Value->SetInputTensorShapes(InputShapes);
 					NNX::FTensorShape OutputShape = Pair.Value->GetOutputTensorShapes()[0];
 
@@ -275,24 +264,20 @@ void FNNENeuralPostProcessing::PrePostProcessPass_RenderThread(FRDGBuilder& Grap
 					int32 NeuralNetworkOutputWidth = OutputShape.Data[3];
 					int32 NeuralNetworkOutputHeight = OutputShape.Data[2];
 
-					FRDGBufferDesc OutputBufferDesc;
-					OutputBufferDesc.Usage = EBufferUsageFlags::UnorderedAccess | EBufferUsageFlags::ShaderResource | EBufferUsageFlags::StructuredBuffer;
-					OutputBufferDesc.BytesPerElement = sizeof(float);
-					OutputBufferDesc.NumElements = NeuralNetworkOutputWidth * NeuralNetworkOutputHeight * 3;
-
+					FRDGBufferDesc OutputBufferDesc = FRDGBufferDesc::CreateBufferDesc(sizeof(float), NeuralNetworkOutputWidth * NeuralNetworkOutputHeight * 3);
 					FRDGBufferRef OutputBuffer = GraphBuilder.CreateBuffer(OutputBufferDesc, *(FString("NNENeuralPostProcessing::NeuralNetowrkOutput_") + FString::FromInt(Pair.Key)));
-					FRDGBufferUAVRef OutputBufferUAV = GraphBuilder.CreateUAV(OutputBuffer);
+					FRDGBufferUAVRef OutputBufferUAV = GraphBuilder.CreateUAV(FRDGBufferUAVDesc(OutputBuffer, PF_R32_FLOAT));
 
 					TArray<NNX::FMLTensorBinding> InputBindings;
 					InputBindings.Add(NNX::FMLTensorBinding::FromRDG(InputBuffer, InputBufferDesc.NumElements * InputBufferDesc.BytesPerElement, 0));
 					TArray<NNX::FMLTensorBinding> OutputBindings;
 					OutputBindings.Add(NNX::FMLTensorBinding::FromRDG(OutputBuffer, OutputBufferDesc.NumElements * OutputBufferDesc.BytesPerElement, 0));
-					// Pair.Value->EnqueueRDG(GraphBuilder, InputBindings, OutputBindings);
+					Pair.Value->EnqueueRDG(GraphBuilder, InputBindings, OutputBindings);
 
 					TNeuralPostProcessingPostStepCS::FParameters* PostStepParameters = GraphBuilder.AllocParameters<TNeuralPostProcessingPostStepCS::FParameters>();
 					PostStepParameters->OutputBufferWidth = NeuralNetworkOutputWidth;
 					PostStepParameters->OutputBufferHeight = NeuralNetworkOutputHeight;
-					PostStepParameters->OutputBuffer = InputBufferUAV;// OutputBufferUAV;
+					PostStepParameters->OutputBuffer = OutputBufferUAV;
 					PostStepParameters->InputTextureWidth = TextureSize.X;
 					PostStepParameters->InputTextureHeight = TextureSize.Y;
 					PostStepParameters->AccumulationBuffer = AccumulationBufferUAV;
