NNE HLSL weights upload + a few fix
Built on top of https://p4-swarm.epicgames.net/reviews/23429965 by christian.regg

Problem:
* Model weights were not uploaded for HLSL engine

Solution:
* Upload the weights + few related fix to make mosaic-9 test green.

Test:
* All tests green under win64 editor, added model with weights to HLSL test when list of op are supported.

diff --git a/Engine/Plugins/Experimental/NNE/Source/NNXCore/Public/NNXTypes.h b/Engine/Plugins/Experimental/NNE/Source/NNXCore/Public/NNXTypes.h
--- a/Engine/Plugins/Experimental/NNE/Source/NNXCore/Public/NNXTypes.h
+++ b/Engine/Plugins/Experimental/NNE/Source/NNXCore/Public/NNXTypes.h
@@ -62,7 +62,7 @@ template <typename DimType, class FinalType> struct TTensorShapeBase
 
 	inline void operator=(const TTensorShapeBase<DimType, FinalType>& OtherShape)
 	{
-		check(Rank() <= MaxRank);
+		check(OtherShape.Rank() <= MaxRank);
 		Data = OtherShape.Data;
 	}
 
@@ -252,8 +252,8 @@ public:
 		const T* DataPtr = reinterpret_cast<const T*>(PreparedData.GetData());
 		const int32 ElemSize = sizeof(T);
 		
-		check(DataSize % ElemSize == 0);
-		return MakeArrayView(DataPtr, DataSize / ElemSize);
+		check(PreparedData.Num() % ElemSize == 0);
+		return MakeArrayView(DataPtr, PreparedData.Num() / ElemSize);
 	}
 
 	void SetShape(const FTensorShape& InShape)
diff --git a/Engine/Plugins/Experimental/NNE/Source/NNXQA/Resources/NNXQATestDesc.json b/Engine/Plugins/Experimental/NNE/Source/NNXQA/Resources/NNXQATestDesc.json
--- a/Engine/Plugins/Experimental/NNE/Source/NNXQA/Resources/NNXQATestDesc.json
+++ b/Engine/Plugins/Experimental/NNE/Source/NNXQA/Resources/NNXQATestDesc.json
@@ -10,7 +10,7 @@
 					"tags": [ "SmokeTest" ],
 					"runtimes": [
 						{ "name": "NNXRuntimeORTDml", "relative_error": 1.0 },
-						{ "name": "NNXRuntimeHlsl",	"skip": true, "comment": "Model with multiple layer not yet implemented."	},
+						{ "name": "NNXRuntimeHlsl",  "skip": true, "comment": "Missing BatchNormalisation operator." },
 						{ "name": "NNXRuntimeDml",  "skip": true, "comment": "Model with multiple layer not yet implemented." }
 					]
 				},
@@ -22,7 +22,7 @@
 							"comment": "Worst relative error 43.367836 % (epsilon 0.001000 %) at position 112558, got -0.003391 expected -0.002365. TODO we need an automatic balance between relative error and absolute error"
 						},
 						{ "name": "NNXRuntimeORTCpu", "absolute_error": 1e-3, "relative_error": 3.0 },
-						{ "name": "NNXRuntimeHlsl",	"skip": true, "comment": "Model with multiple layer not yet implemented."	},
+						{ "name": "NNXRuntimeHlsl", "absolute_error": 0.03, "relative_error": 55.0, "comment": "TODO check, errors are very high." },
 						{ "name": "NNXRuntimeDml",  "skip": true, "comment": "Model with multiple layer not yet implemented." }
 					]
 				}
@@ -37,7 +37,7 @@
 					"tags": [ "SmokeTest" ],
 					"runtimes": [
 						{ "name": "NNXRuntimeORTDml", "relative_error": 0.02 },
-						{ "name": "NNXRuntimeHlsl",	"skip": true, "comment": "Model with multiple layer not yet implemented."	},
+						{ "name": "NNXRuntimeHlsl", "relative_error": 0.02 },
 						{ "name": "NNXRuntimeDml",  "skip": true, "comment": "Model with multiple layer not yet implemented." }
 					]
 				},
@@ -46,17 +46,15 @@
 					"tags": [ "SmokeTest" ],
 					"runtimes": [
 						{ "name": "NNXRuntimeORTDml", "relative_error": 0.04 },
-						{ "name": "NNXRuntimeHlsl",	"skip": true, "comment": "Model with multiple layer not yet implemented."	},
+						{ "name": "NNXRuntimeHlsl",  "skip": true, "comment": "Missing Reshape operator." },
 						{ "name": "NNXRuntimeDml",  "skip": true, "comment": "Model with multiple layer not yet implemented." }
 					]
 				},
 				{
 					"target": "VertexDeltaModel",
 					"runtimes": [
-						{ "name": "NNXRuntimeORTDml", "absolute_error": 1e-4, "relative_error": 10.0,
-							"comment": "TODO check, relative error is very high."
-						},
-						{ "name": "NNXRuntimeHlsl",	"skip": true, "comment": "Model with multiple layer not yet implemented."	},
+						{ "name": "NNXRuntimeORTDml", "absolute_error": 1e-4, "relative_error": 10.0, "comment": "TODO check, relative error is very high."	},
+						{ "name": "NNXRuntimeHlsl", "absolute_error": 1e-4, "relative_error": 3.0 },
 						{ "name": "NNXRuntimeDml",  "skip": true, "comment": "Model with multiple layer not yet implemented." }
 					]
 				}
@@ -70,7 +68,6 @@
 					"target": "OneLayerWithWeights",
 					"tags": [ "SmokeTest" ],
 					"runtimes": [
-						{ "name": "NNXRuntimeHlsl",	"skip": true, "comment": "Weights from model not yet implemented."	},
 						{ "name": "NNXRuntimeDml",  "skip": true, "comment": "Weights from model not yet implemented." }
 					]
 				},
@@ -632,11 +629,11 @@
 		{
 			"name": "Upsample.ScaleFromWeight",
 			"datasets": [
-				{ "inputs":  [ { "shape": [ 1 ] } ],             "weights": [ { "shape": [ 1 ], "source": [ "3"    ],                  "type": "float" } ], "outputs": [ { "shape": [ 3 ]} ] },
-				{ "inputs":  [ { "shape": [ 2 ] } ],             "weights": [ { "shape": [ 1 ], "source": [ "2"    ],                  "type": "float" } ], "outputs": [ { "shape": [ 4 ]} ] },
-				{ "inputs":  [ { "shape": [ 3 ] } ],             "weights": [ { "shape": [ 1 ], "source": [ "1"    ],                  "type": "float" } ], "outputs": [ { "shape": [ 3 ]} ] },
-				{ "inputs":  [ { "shape": [ 1, 128, 50, 55 ] } ],    "weights": [ { "shape": [ 4 ], "source": [ "1", "1", "2", "2" ],      "type": "float" } ], "outputs": [ { "shape": [ 1, 128, 100, 110 ]} ] },
-				{ "inputs":  [ { "shape": [ 2, 1, 3, 1, 2 ] } ], "weights": [ { "shape": [ 5 ], "source": [ "1", "2", "2", "1", "2" ], "type": "float" } ], "outputs": [ { "shape": [ 2, 2, 6, 1, 4 ]} ] }
+				{ "inputs":  [ { "shape": [ 1 ] } ],              "weights": [ { "shape": [ 1 ], "source": [ "3"    ],                  "type": "float" } ], "outputs": [ { "shape": [ 3 ]} ] },
+				{ "inputs":  [ { "shape": [ 2 ] } ],              "weights": [ { "shape": [ 1 ], "source": [ "2"    ],                  "type": "float" } ], "outputs": [ { "shape": [ 4 ]} ] },
+				{ "inputs":  [ { "shape": [ 3 ] } ],              "weights": [ { "shape": [ 1 ], "source": [ "1"    ],                  "type": "float" } ], "outputs": [ { "shape": [ 3 ]} ] },
+				{ "inputs":  [ { "shape": [ 1, 128, 50, 55 ] } ], "weights": [ { "shape": [ 4 ], "source": [ "1", "1", "2", "2" ],      "type": "float" } ], "outputs": [ { "shape": [ 1, 128, 100, 110 ]} ] },
+				{ "inputs":  [ { "shape": [ 2, 1, 3, 1, 2 ] } ],  "weights": [ { "shape": [ 5 ], "source": [ "1", "2", "2", "1", "2" ], "type": "float" } ], "outputs": [ { "shape": [ 2, 2, 6, 1, 4 ]} ] }
 			]
 		},
 		{
diff --git a/Engine/Plugins/Experimental/NNE/Source/NNXRuntimeRDG/Private/Hlsl/NNIRuntimeRDGInstanceNormalization.cpp b/Engine/Plugins/Experimental/NNE/Source/NNXRuntimeRDG/Private/Hlsl/NNIRuntimeRDGInstanceNormalization.cpp
--- a/Engine/Plugins/Experimental/NNE/Source/NNXRuntimeRDG/Private/Hlsl/NNIRuntimeRDGInstanceNormalization.cpp
+++ b/Engine/Plugins/Experimental/NNE/Source/NNXRuntimeRDG/Private/Hlsl/NNIRuntimeRDGInstanceNormalization.cpp
@@ -120,16 +120,16 @@ namespace UE::NNIRuntimeRDG::Private::Hlsl
 			Parameters->Bias = GraphBuilder.CreateSRV(FRDGBufferSRVDesc(Bias.GetBuffer(), PF_R32_FLOAT));
 			Parameters->Output = GraphBuilder.CreateUAV(FRDGBufferUAVDesc(Output.GetBuffer(), PF_R32_FLOAT));
 
-			TInstanceNormalizationCS::FPermutationDomain PermutationVector;
-			PermutationVector.Set<TInstanceNormalizationCS::FInstanceNormalizationAlgorithm>(Algorithm);
-			TShaderMapRef<TInstanceNormalizationCS> ComputeShader(GetGlobalShaderMap(GMaxRHIFeatureLevel), PermutationVector);
-
-			if (Algorithm == EInstanceNormalizationAlgorithm::MAX)
+			EInstanceNormalizationAlgorithm DispatchAlgorithm = Algorithm;
+			if (DispatchAlgorithm == EInstanceNormalizationAlgorithm::MAX)
 			{
-				TInstanceNormalizationCS::GetAlgorithm(*Parameters);
+				DispatchAlgorithm = TInstanceNormalizationCS::GetAlgorithm(*Parameters);
 			}
-
-			FIntVector ThreadGroupCount = TInstanceNormalizationCS::GetGroupCount(*Parameters, Algorithm);
+			
+			TInstanceNormalizationCS::FPermutationDomain PermutationVector;
+			PermutationVector.Set<TInstanceNormalizationCS::FInstanceNormalizationAlgorithm>(DispatchAlgorithm);
+			TShaderMapRef<TInstanceNormalizationCS> ComputeShader(GetGlobalShaderMap(GMaxRHIFeatureLevel), PermutationVector);
+			FIntVector ThreadGroupCount = TInstanceNormalizationCS::GetGroupCount(*Parameters, DispatchAlgorithm);
 
 			RDG_EVENT_SCOPE(GraphBuilder, "NNI.Operator.Hlsl.InstanceNormalization");
 			RDG_GPU_STAT_SCOPE(GraphBuilder, FNNIOperatorInstanceNormalization);
diff --git a/Engine/Plugins/Experimental/NNE/Source/NNXRuntimeRDG/Private/NNERuntimeRDGModelHLSL.cpp b/Engine/Plugins/Experimental/NNE/Source/NNXRuntimeRDG/Private/NNERuntimeRDGModelHLSL.cpp
new file mode 100644
--- /dev/null
+++ b/Engine/Plugins/Experimental/NNE/Source/NNXRuntimeRDG/Private/NNERuntimeRDGModelHLSL.cpp
@@ -0,0 +1,443 @@
+// Copyright Epic Games, Inc. All Rights Reserved.
+
+#include "NNERuntimeRDGModelHLSL.h"
+#include "NNXRuntimeHLSLOp.h"
+#include "RenderGraphBuilder.h"
+#include "RenderGraphUtils.h"
+
+namespace UE::NNIRuntimeRDG::Private::Hlsl
+{
+
+namespace
+{
+
+NNX::FMLOperatorHlsl* OpCreate(const FString& OpName, TConstArrayView<NNX::FTensorDesc> InputTensorDescs, TConstArrayView<NNX::FTensorDesc> OutputTensorDescs, const UE::NNECore::FAttributeMap& AttributeMap)
+{
+	NNX::FMLOperatorRegistryHlsl::OperatorCreateFunc CreateFn = NNX::FMLOperatorRegistryHlsl::Get()->OpFind(OpName);
+
+	if (!CreateFn)
+	{
+		UE_LOG(LogNNX, Warning, TEXT("Hlsl MLOperatorRegistry failed to find operator:%s"), *OpName);
+		return nullptr;
+	}
+
+	NNX::FMLOperatorHlsl* Op = CreateFn();
+
+	if (!Op->Initialize(InputTensorDescs, OutputTensorDescs, AttributeMap))
+	{
+		UE_LOG(LogNNX, Warning, TEXT("Hlsl engine: Error initializing operator:%s"), *OpName);
+		delete Op;
+		return nullptr;
+	}
+
+	return Op;
+}
+
+template<class OperatorRef>
+void InternAddDispatchOps_RenderThread(FRDGBuilder& GraphBuilder, TConstArrayView<FTensorHLSLRef> AllTensorHLSLRefs, TConstArrayView<TArray<uint32>> OperatorInputTensorIndices,
+	TConstArrayView<TArray<uint32>> OperatorOutputTensorIndices, TArrayView<OperatorRef> Operators)
+{
+	static constexpr int32 MaxExpectedInput = 10;
+	TArray<NNX::FTensorRDGRef, TInlineAllocator<MaxExpectedInput>> InputTensors;
+
+	static constexpr int32 MaxExpectedOutput = 2;
+	TArray<NNX::FTensorRDGRef, TInlineAllocator<MaxExpectedOutput>> OutputTensors;
+
+	// Add passes for all operators
+	for (int32 Idx = 0; Idx < Operators.Num(); ++Idx)
+	{
+		InputTensors.Empty();
+		for (int32 i : OperatorInputTensorIndices[Idx])
+		{
+			InputTensors.Add(AllTensorHLSLRefs[i]);
+		}
+		OutputTensors.Empty();
+		for (int32 i : OperatorOutputTensorIndices[Idx])
+		{
+			OutputTensors.Add(AllTensorHLSLRefs[i]);
+		}
+
+		Operators[Idx]->Dispatch(GraphBuilder, InputTensors, OutputTensors);
+	}
+}
+
+int ApplyBinding(TArray<FTensorHLSL>& OutTensorsHLSL, TConstArrayView<NNX::FMLTensorBinding> InBindings, bool bIsInput)
+{
+	check(OutTensorsHLSL.Num() == InBindings.Num());
+	
+	for (int32 Idx = 0; Idx < InBindings.Num(); ++Idx)
+	{
+		const NNX::FMLTensorBinding& Binding = InBindings[Idx];
+		FTensorHLSL& Tensor = OutTensorsHLSL[Idx];
+
+		if (Binding.BindingType == NNX::EMLTensorBindingDataType::CPUMemory)
+		{
+			if (bIsInput) Tensor.SetUploadBuffer(Binding.CpuMemory);
+			else Tensor.SetDownloadBuffer(Binding.CpuMemory);
+		}
+		else if (Binding.BindingType == NNX::EMLTensorBindingDataType::RDGBuffer)
+		{
+			Tensor.SetBuffer(Binding.Buffer);
+		}
+		else
+		{
+			// Unsupported tensor binding type
+			return Idx;
+		}
+	}
+
+	return 0;
+}
+
+bool ApplyWeights(FRDGBuilder& GraphBuilder, TArray<FTensorHLSL>& OutTensorsHLSL, TArray<TRefCountPtr<FRDGPooledBuffer>> ExternalWeightsRDG)
+{
+	check(OutTensorsHLSL.Num() == ExternalWeightsRDG.Num());
+
+	for (int32 Idx = 0; Idx < ExternalWeightsRDG.Num(); ++Idx)
+	{
+		const TRefCountPtr<FRDGPooledBuffer>& PooledBuffer = ExternalWeightsRDG[Idx];
+		FTensorHLSL& Tensor = OutTensorsHLSL[Idx];
+		FRDGBufferRef Buffer = GraphBuilder.RegisterExternalBuffer(PooledBuffer);
+		Tensor.SetBuffer(Buffer);
+	}
+
+	return true;
+}
+
+}
+
+bool FModel::Init(TConstArrayView<uint8> ModelData)
+{
+	check(ModelData.Num() > 0);
+	FMLRuntimeFormat	Format;
+
+	if (!LoadModel(ModelData, Format))
+	{
+		return false;
+	}
+
+	// Create HLSL tensor and upload to GPU
+	PrepareWeights();
+
+	// Loop over all operators in the model and create them
+	for (int32 Idx = 0; Idx < Format.Operators.Num(); ++Idx)
+	{
+		const FString TypeName = Format.Operators[Idx].TypeName;
+
+		TArray<NNX::FTensorDesc> Inputs;
+		TArray<NNX::FTensorDesc> Outputs;
+		UE::NNECore::FAttributeMap AttributeMap;
+
+		for (int32 InputTensorIndex : Format.Operators[Idx].InTensors)
+		{
+			Inputs.Emplace(AllSymbolicTensorDescs[InputTensorIndex]);
+		}
+		for (int32 OutputTensorIndex : Format.Operators[Idx].OutTensors)
+		{
+			Outputs.Emplace(AllSymbolicTensorDescs[OutputTensorIndex]);
+		}
+		for (const FMLFormatAttributeDesc& Desc : Format.Operators[Idx].Attributes)
+		{
+			AttributeMap.SetAttribute(Desc.Name, Desc.Value);
+		}
+
+		NNX::FMLOperatorHlsl* Op = OpCreate(TypeName, Inputs, Outputs, AttributeMap);
+
+		if (!Op) //Op.Shader.IsNull())
+		{
+			UE_LOG(LogNNX, Warning, TEXT("Failed to create operator:%s"), *TypeName);
+
+			// TODO: Cleanup operators
+			return false;
+		}
+
+		Operators.Add(Op);
+	}
+
+	return true;
+}
+
+int FModel::SetInputTensorShapes(TConstArrayView<NNX::FTensorShape> InputShapes)
+{
+	int Res = NNX::FMLInferenceModelRDG::SetInputTensorShapes(InputShapes);
+	if (Res < 0) return Res;
+
+	auto ConvertTensors = [this] (TArray<FTensorHLSLRef> &TensorRefs, TArray<FTensorHLSL> &Tensors, const NNX::FTensorRDGArray &TensorRDGs, const TArray<int32>& TensorIndices)
+	{
+		Tensors.Empty();
+		Tensors.SetNum(TensorRDGs.Num());
+
+		for (int32 i = 0; i < TensorIndices.Num(); ++i)
+		{
+			const int32 Idx = TensorIndices[i];
+			Tensors[i] = FTensorHLSL(TensorRDGs[i]);
+			AllTensorHLSLRefs[Idx]= &Tensors[i];
+		}
+	};
+
+	AllTensorHLSLRefs.Empty();
+	AllTensorHLSLRefs.SetNumUninitialized(AllTensorRDGs.Num());
+	ConvertTensors(AllTensorHLSLRefs, InputTensorHLSLs, InputTensorRDGs, InputTensorIndices);
+	ConvertTensors(AllTensorHLSLRefs, OutputTensorHLSLs, OutputTensorRDGs, OutputTensorIndices);
+	ConvertTensors(AllTensorHLSLRefs, IntermediateTensorHLSLs, IntermediateTensorRDGs, IntermediateTensorIndices);
+	for (int32 i = 0; i < WeightTensorIndices.Num(); ++i)
+	{
+		const int32 Idx = WeightTensorIndices[i];
+		AllTensorHLSLRefs[Idx] = &WeightTensorHLSLs[i];
+	}
+
+	return 0;
+}
+
+int FModel::RunSync(TConstArrayView<NNX::FMLTensorBinding> InInputBindings, TConstArrayView<NNX::FMLTensorBinding> InOutputBindings)
+{
+	// Verify the model inputs were prepared
+	if (InputTensorShapes.Num() == 0)
+	{
+		UE_LOG(LogNNX, Error, TEXT("Run(): Input shapes are not set, please call SetInputTensorShapes."));
+		return -1;
+	}
+	
+	int Res = 0;
+	FEvent* Signal = FGenericPlatformProcess::GetSynchEventFromPool(false);
+
+	ENQUEUE_RENDER_COMMAND(FMLInferenceModel_Run)
+	(
+		[&Signal, &Res, this, InInputBindings, InOutputBindings](FRHICommandListImmediate& RHICmdList)
+		{
+			TOptional<ERHIPipeline>		Pipeline = RHICmdList.GetPipeline();
+
+			if (Pipeline == ERHIPipeline::None)
+			{
+				RHICmdList.SwitchPipeline(ERHIPipeline::Graphics);
+			}
+
+			FRDGBuilder	RDGBuilder(RHICmdList);
+
+			Res = EnqueueRDG(RDGBuilder, InInputBindings, InOutputBindings);
+			if (Res == 0)
+			{
+				RDGBuilder.Execute();
+
+				// FIXME: Using BlockUntilGPUIdle() prevents hang on Linux
+				// FIXME: Adapt to redesigned readback API (UE 5.2)
+				RHICmdList.BlockUntilGPUIdle();
+
+				for (FTensorHLSLRef Tensor : AllTensorHLSLRefs)
+				{
+					Tensor->Resolve();
+				}
+			}
+			
+			Signal->Trigger();
+		}
+	);
+
+	// We need to wait for render thread to finish
+	Signal->Wait();
+
+	FGenericPlatformProcess::ReturnSynchEventToPool(Signal);
+
+	return Res;
+}
+
+int FModel::EnqueueRDG(FRDGBuilder& GraphBuilder, TConstArrayView<NNX::FMLTensorBinding> InInputBindings, TConstArrayView<NNX::FMLTensorBinding> InOutputBindings)
+{
+	check(IsInRenderingThread());
+
+	int Res;
+
+	// Verify the model inputs were prepared
+	if (InputTensorShapes.Num() == 0)
+	{
+		UE_LOG(LogNNX, Error, TEXT("EnqueueRDG(): Input shapes are not set, please call SetInputTensorShapes."));
+		return -1;
+	}
+
+	Res = ApplyBinding(InputTensorHLSLs, InInputBindings, true);
+	if (Res != 0)
+	{
+		UE_LOG(LogNNX, Warning, TEXT("Invalid input tensor binding type for tensor index:%d"), Res);
+		return -1;
+	}
+
+	Res = ApplyBinding(OutputTensorHLSLs, InOutputBindings, false);
+	if (Res != 0)
+	{
+		UE_LOG(LogNNX, Warning, TEXT("Invalid output tensor binding type for tensor index:%d"), Res);
+		return -1;
+	}
+
+	if (!ApplyWeights(GraphBuilder, WeightTensorHLSLs, WeightsExternalRDGResources))
+	{
+		UE_LOG(LogNNX, Warning, TEXT("Could not register the weights for graph execution."));
+		return -1;
+	}
+
+	// Create RDG buffers where not yet present
+	for (FTensorHLSLRef Tensor : AllTensorHLSLRefs)
+	{
+		if (!Tensor->HasBuffer())
+		{
+			FRDGBufferDesc BufferDesc = FRDGBufferDesc::CreateBufferDesc(Tensor->GetElemByteSize(), Tensor->GetVolume());
+			if (Tensor->HasDownloadBuffer())
+			{
+				BufferDesc.Usage = EBufferUsageFlags(BufferDesc.Usage | BUF_SourceCopy);
+			}
+
+			const FRDGBufferRef TensorBuffer = GraphBuilder.CreateBuffer(BufferDesc, *Tensor->GetName(), ERDGBufferFlags::None);
+			
+			Tensor->SetBuffer(TensorBuffer);
+		}
+	}
+
+	for (FTensorHLSLRef Tensor : AllTensorHLSLRefs)
+	{
+		Tensor->EnqueueUploadRdg(GraphBuilder);
+	}
+
+	InternAddDispatchOps_RenderThread<NNX::FMLOperatorHlsl*>(GraphBuilder, AllTensorHLSLRefs, OperatorInputTensorIndices, OperatorOutputTensorIndices, Operators);
+
+	for (FTensorHLSLRef Tensor : AllTensorHLSLRefs)
+	{
+		Tensor->EnqueueDownloadRdg(GraphBuilder, bUseManualTransitions);
+	}
+
+	return 0;
+}
+
+int FModel::PrepareTensorShapesAndData()
+{
+	check(AllTensorRDGs.Num() == AllSymbolicTensorDescs.Num());
+	
+	if (Operators.Num() == 0)
+	{
+		UE_LOG(LogNNX, Warning, TEXT("No operators in model"));
+		return -1;
+	}
+
+	// Run model preparation (including shape inference) on all operators
+	// This loop could be abstracted to a different engine/system as it apply on FTensorRef & IPrepareOperator witch are RDG agnostics.
+	static constexpr int32 MaxExpectedInput = 10;
+	TArray<NNX::FTensorRef, TInlineAllocator<MaxExpectedInput>> InputTensors;
+	TArray<NNX::FTensorRef> OutputTensors;
+	TArray<bool> AllInitializedTensors;
+
+	checkCode(
+		AllInitializedTensors.Init(false, AllSymbolicTensorDescs.Num());
+		for (int32 Idx : InputTensorIndices)
+		{
+			AllInitializedTensors[Idx] = true;
+		}
+		for (int32 Idx : WeightTensorIndices)
+		{
+			AllInitializedTensors[Idx] = true;
+		}
+	);
+
+	// Run model preparation (including shape inference) on all operators
+	// This loop could be abstracted to a different system as it apply on FTensorRef & IPrepareOperator witch are RDG agnostics.
+	for (int32 Idx = 0; Idx < Operators.Num(); ++Idx)
+	{
+		InputTensors.Empty();
+		OutputTensors.Empty();
+
+		//Operator inputs
+		for (int32 i : OperatorInputTensorIndices[Idx])
+		{
+			checkf(AllInitializedTensors[i] == true, TEXT("Input tensor %d for operator %d should have been initialized."), i, Idx);
+			InputTensors.Emplace(AllTensorRDGs[i]);
+		}
+		//Operator outputs
+		for (int32 i : OperatorOutputTensorIndices[Idx])
+		{
+			OutputTensors.Emplace(AllTensorRDGs[i]);
+			checkf(AllInitializedTensors[i] == false, TEXT("Output tensor %d for operator %d should not have been initialized yet."), i, Idx);
+			checkCode(AllInitializedTensors[i] = true);
+		}
+
+		const NNX::FMLOperatorHlsl* Op = Operators[Idx];
+
+		if (Op->PrepareOutputs(InputTensors, OutputTensors) != 0)
+		{
+			//Operator could not prepare the output tensors, meaning we can't allocate
+			//output buffer before running the model. This engine does not support this.
+			UE_LOG(LogNNX, Warning, TEXT("Could not deduce tensor shapes for this model during shape inference, HLSL engine wont support the model as it need to precompute all shapes for performance reasons."));
+			AllTensorRDGs.Empty();
+			return -1;
+		}
+	}
+
+	checkCode(
+		for (int i = 0; i < AllInitializedTensors.Num(); ++i)
+		{
+			checkf(AllInitializedTensors[i], TEXT("Tensor at index %d, was not initialized by model preparation."));
+		};
+	);
+
+	return 0;
+}
+
+bool FModel::PrepareWeights()
+{
+	if (!WeightsExternalRDGResources.IsEmpty())
+	{
+		check(WeightsExternalRDGResources.Num() == WeightTensorHLSLs.Num())
+		return true;
+	}
+
+	//Convert to HLSL tensor
+	//This copy the weight. To be improved.
+	check(WeightTensorHLSLs.IsEmpty());
+	WeightTensorHLSLs.SetNum(WeightTensorRDGs.Num());
+	for (int32 Idx = 0; Idx < WeightTensorRDGs.Num(); ++Idx)
+	{
+		WeightTensorHLSLs[Idx] = FTensorHLSL(WeightTensorRDGs[Idx]);
+	}
+
+	//Upload to GPU
+	FEvent* Signal = FGenericPlatformProcess::GetSynchEventFromPool(false);
+	WeightsExternalRDGResources.SetNum(WeightTensorHLSLs.Num());
+
+	ENQUEUE_RENDER_COMMAND(FModel_PrepareWeights)
+	(
+		[&Signal, this](FRHICommandListImmediate& RHICmdList)
+		{
+			TOptional<ERHIPipeline> Pipeline = RHICmdList.GetPipeline();
+			if (Pipeline == ERHIPipeline::None)
+			{
+				RHICmdList.SwitchPipeline(ERHIPipeline::Graphics);
+			}
+
+			FRDGBuilder	RDGBuilder(RHICmdList);
+
+			for (int32 i = 0; i < WeightTensorHLSLs.Num(); ++i)
+			{
+				FTensorHLSL& Tensor = WeightTensorHLSLs[i];
+				check(!Tensor.HasBuffer());
+				check(Tensor.HasPreparedData());
+				FRDGBufferDesc BufferDesc = FRDGBufferDesc::CreateBufferDesc(Tensor.GetElemByteSize(), Tensor.GetVolume());
+				const FRDGBufferRef TransientRDGBuffer = RDGBuilder.CreateBuffer(BufferDesc, *Tensor.GetName(), ERDGBufferFlags::None);
+				WeightsExternalRDGResources[i] = RDGBuilder.ConvertToExternalBuffer(TransientRDGBuffer);
+				Tensor.SetBuffer(TransientRDGBuffer);
+				Tensor.EnqueueUploadRdg(RDGBuilder);
+			}
+
+			RDGBuilder.Execute();
+
+			//To prevent any problem if model is released before upload is done to the GPU. To be improved.
+			RHICmdList.BlockUntilGPUIdle();
+
+			Signal->Trigger();
+		}
+	);
+
+	// We need to wait for render thread to finish
+	Signal->Wait();
+
+	FGenericPlatformProcess::ReturnSynchEventToPool(Signal);
+
+	return true;
+}
+
+} // namespace UE::NNIRuntimeRDG::Private::Hlsl
\ No newline at end of file
diff --git a/Engine/Plugins/Experimental/NNE/Source/NNXRuntimeRDG/Private/NNERuntimeRDGModelHLSL.h b/Engine/Plugins/Experimental/NNE/Source/NNXRuntimeRDG/Private/NNERuntimeRDGModelHLSL.h
new file mode 100644
--- /dev/null
+++ b/Engine/Plugins/Experimental/NNE/Source/NNXRuntimeRDG/Private/NNERuntimeRDGModelHLSL.h
@@ -0,0 +1,55 @@
+// Copyright Epic Games, Inc. All Rights Reserved.
+
+#pragma once
+
+#include "NNXRuntimeRDG.h"
+#include "NNXTypes.h"
+#include "NNERuntimeRDGTensorHLSL.h"
+
+class FRDGBuilder;
+namespace NNX { struct FMLOperatorHlsl; }
+
+namespace UE::NNIRuntimeRDG::Private::Hlsl
+{
+
+class FModel : public NNX::FMLInferenceModelRDG
+{
+	
+public:
+
+	~FModel() = default;
+
+	bool Init(TConstArrayView<uint8> ModelData);
+
+	virtual int SetInputTensorShapes(TConstArrayView<NNX::FTensorShape> InputShapes) override;
+
+	/**
+	 * Run the inference model (synchronous version)
+	 */
+	virtual int RunSync(TConstArrayView<NNX::FMLTensorBinding> InputBindings, TConstArrayView<NNX::FMLTensorBinding> OutputBindings) override;
+	
+	/**
+	 * Enqueue operators to RDG, the caller will run the GraphBuilder.Execute()
+	 */
+	virtual int EnqueueRDG(FRDGBuilder& GraphBuilder, TConstArrayView<NNX::FMLTensorBinding> InputBindings, TConstArrayView<NNX::FMLTensorBinding> OutputBindings) override;
+
+protected:
+
+	virtual void AddDispatchOps_RenderThread(FRDGBuilder& GraphBuilder) override { };
+	virtual int PrepareTensorShapesAndData() override;
+
+	bool PrepareWeights();
+
+private:
+
+	TArray<NNX::FMLOperatorHlsl*>	Operators;
+
+	TArray<FTensorHLSLRef> AllTensorHLSLRefs;
+	TArray<FTensorHLSL> InputTensorHLSLs;
+	TArray<FTensorHLSL> OutputTensorHLSLs;
+	TArray<FTensorHLSL> IntermediateTensorHLSLs;
+	TArray<FTensorHLSL> WeightTensorHLSLs;
+	TArray<TRefCountPtr<FRDGPooledBuffer>> WeightsExternalRDGResources;
+};
+
+} // namespace UE::NNIRuntimeRDG::Private::Hlsl
\ No newline at end of file
diff --git a/Engine/Plugins/Experimental/NNE/Source/NNXRuntimeRDG/Private/NNERuntimeRDGTensorHLSL.cpp b/Engine/Plugins/Experimental/NNE/Source/NNXRuntimeRDG/Private/NNERuntimeRDGTensorHLSL.cpp
new file mode 100644
--- /dev/null
+++ b/Engine/Plugins/Experimental/NNE/Source/NNXRuntimeRDG/Private/NNERuntimeRDGTensorHLSL.cpp
@@ -0,0 +1,69 @@
+// Copyright Epic Games, Inc. All Rights Reserved.
+
+#include "NNERuntimeRDGTensorHLSL.h"
+#include "RenderGraphBuilder.h"
+
+namespace UE::NNIRuntimeRDG::Private::Hlsl
+{
+
+void FTensorHLSL::EnqueueUploadRdg(FRDGBuilder& GraphBuilder)
+{
+	if (!HasPreparedData() && !HasUploadBuffer()) return;
+	
+	check(HasPreparedData() != HasUploadBuffer());
+	check(HasBuffer());
+
+	FCpuBufferRef CPUBuffer = HasUploadBuffer() ? GetUploadBuffer() : PreparedData.GetData();
+	GraphBuilder.QueueBufferUpload(GetBuffer(), CPUBuffer, GetDataSize(), ERDGInitialDataFlags::NoCopy);
+}
+
+void FTensorHLSL::EnqueueDownloadRdg(FRDGBuilder& GraphBuilder, bool bUseManualTransitions)
+{
+	if (!HasDownloadBuffer()) return;
+
+	check(HasBuffer());
+	check(!Readback);
+
+	Readback = MakeUnique<FRHIGPUBufferReadback>(FName(TEXT("FMLTensorReadback_") + GetName()));
+
+	FMLTensorReadbackParameters* TensorReadbackParams = GraphBuilder.AllocParameters<FMLTensorReadbackParameters>();
+	TensorReadbackParams->Buffer = GetBuffer();
+
+	GraphBuilder.AddPass(
+		RDG_EVENT_NAME("FMLInferenceModelAddTensorReadback:%s", *GetName()),
+		TensorReadbackParams,
+		ERDGPassFlags::Readback | ERDGPassFlags::NeverCull,
+		[this, bUseManualTransitions, TensorReadbackParams](FRHICommandListImmediate& RHICmdList)
+		{
+			FRHIBuffer* OutputBuffer = TensorReadbackParams->Buffer->GetRHI();
+
+			// TODO: FIXME: We need to transition the resources for DirectML
+			if (bUseManualTransitions)
+			{
+				FRHITransitionInfo Transitions[] =
+				{
+					FRHITransitionInfo(OutputBuffer, ERHIAccess::UAVCompute, ERHIAccess::CopySrc)
+				};
+
+				RHICmdList.Transition(MakeArrayView(Transitions, UE_ARRAY_COUNT(Transitions)));
+				RHICmdList.SubmitCommandsHint();
+			}
+
+			Readback->EnqueueCopy(RHICmdList, OutputBuffer, GetDataSize());
+		}
+	);
+}
+
+void FTensorHLSL::Resolve()
+{
+	if (!Readback) return;
+
+	const void* BuffData = Readback->Lock(GetDataSize());
+	check(BuffData);
+
+	FMemory::Memcpy(GetDownloadBuffer(), BuffData, GetDataSize());
+
+	Readback->Unlock();
+}
+
+} // namespace UE::NNIRuntimeRDG::Private::Hlsl
\ No newline at end of file
diff --git a/Engine/Plugins/Experimental/NNE/Source/NNXRuntimeRDG/Private/NNERuntimeRDGTensorHLSL.h b/Engine/Plugins/Experimental/NNE/Source/NNXRuntimeRDG/Private/NNERuntimeRDGTensorHLSL.h
new file mode 100644
--- /dev/null
+++ b/Engine/Plugins/Experimental/NNE/Source/NNXRuntimeRDG/Private/NNERuntimeRDGTensorHLSL.h
@@ -0,0 +1,61 @@
+// Copyright Epic Games, Inc. All Rights Reserved.
+
+#pragma once
+
+#include "NNXRuntimeRDG.h"
+#include "NNXTypes.h"
+#include "RHIGPUReadback.h"
+
+using FCpuBufferRef = void*;
+
+class FRDGBuilder;
+
+namespace UE::NNIRuntimeRDG::Private::Hlsl
+{
+
+class FTensorHLSL : public NNX::FTensorRDG
+{
+
+private:
+
+	FCpuBufferRef UploadBuffer{};
+	FCpuBufferRef DownloadBuffer{};
+
+	TUniquePtr<FRHIGPUBufferReadback> Readback;
+
+public:
+
+	FTensorHLSL() = default;
+
+	FTensorHLSL(const NNX::FTensorRDG &TensorRDG)
+		: FTensorHLSL(TensorRDG.GetName(), TensorRDG.GetDataType(), TensorRDG.GetShape())
+	{
+		this->PreparedData = TensorRDG.GetPreparedData<uint8>();//This copy the data, we should find another way
+	}
+	
+	FTensorHLSL(const FString& Name, EMLTensorDataType DataType, const NNX::FTensorShape& Shape)
+	{
+		this->Name = Name;
+		this->DataType = DataType;
+		this->Shape = Shape;
+		Volume = Shape.Volume();
+		DataSize = (uint64)NNX::GetTensorDataTypeSizeInBytes(DataType) * Volume;
+	}
+
+	bool HasUploadBuffer() const { return UploadBuffer != FCpuBufferRef{}; }
+	void SetUploadBuffer(FCpuBufferRef Inbuffer){ check(!HasPreparedData()); UploadBuffer = Inbuffer; }
+	FCpuBufferRef GetUploadBuffer() const { return UploadBuffer; }
+
+	bool HasDownloadBuffer() const { return DownloadBuffer != FCpuBufferRef{}; }
+	void SetDownloadBuffer(FCpuBufferRef Inbuffer){ DownloadBuffer = Inbuffer; }
+	FCpuBufferRef GetDownloadBuffer() const { return DownloadBuffer; }
+
+	void EnqueueUploadRdg(FRDGBuilder& GraphBuilder);
+	void EnqueueDownloadRdg(FRDGBuilder& GraphBuilder, bool bUseManualTransitions);
+
+	void Resolve();
+};
+
+using FTensorHLSLRef = FTensorHLSL*;
+
+} // namespace UE::NNIRuntimeRDG::Private::Hlsl
\ No newline at end of file
diff --git a/Engine/Plugins/Experimental/NNE/Source/NNXRuntimeRDG/Private/NNXRuntimeHLSL.cpp b/Engine/Plugins/Experimental/NNE/Source/NNXRuntimeRDG/Private/NNXRuntimeHLSL.cpp
--- a/Engine/Plugins/Experimental/NNE/Source/NNXRuntimeRDG/Private/NNXRuntimeHLSL.cpp
+++ b/Engine/Plugins/Experimental/NNE/Source/NNXRuntimeRDG/Private/NNXRuntimeHLSL.cpp
@@ -1,24 +1,25 @@
 // Copyright Epic Games, Inc. All Rights Reserved.
 
-#include "NNXInferenceModel.h"
-#include "NNXRuntimeFormat.h"
-#include "NNXModelOptimizer.h"
-#include "NNECoreAttributeMap.h"
 #include "Hlsl/NNIRuntimeRDGConv.h"
+#include "Hlsl/NNIRuntimeRDGConv.h"
+#include "Hlsl/NNIRuntimeRDGConvTranspose.h"
 #include "Hlsl/NNIRuntimeRDGElementWiseBinary.h"
 #include "Hlsl/NNIRuntimeRDGElementWiseUnary.h"
 #include "Hlsl/NNIRuntimeRDGElementWiseVariadic.h"
-#include "Hlsl/NNIRuntimeRDGConv.h"
-#include "Hlsl/NNIRuntimeRDGConvTranspose.h"
 #include "Hlsl/NNIRuntimeRDGGemm.h"
-#include "Hlsl/NNIRuntimeRDGMatMul.h"
 #include "Hlsl/NNIRuntimeRDGInstanceNormalization.h"
 #include "Hlsl/NNIRuntimeRDGPad.h"
 #include "Hlsl/NNIRuntimeRDGUpsample.h"
+#include "Hlsl/NNIRuntimeRDGMatMul.h"
+#include "NNECoreAttributeMap.h"
+#include "NNERuntimeRDGModelHLSL.h"
+#include "NNXInferenceModel.h"
+#include "NNXModelOptimizer.h"
+#include "NNXRuntimeFormat.h"
 #include "NNXRuntimeHLSLHelper.h"
 #include "NNXRuntimeRDG.h"
 
-namespace NNX
+namespace UE::NNIRuntimeRDG::Private::Hlsl
 {
 
 #define NNX_RUNTIME_HLSL_NAME TEXT("NNXRuntimeHlsl")
@@ -26,226 +27,7 @@ namespace NNX
 //
 //
 //
-class FMLInferenceModelHlsl : public FMLInferenceModelRDG
-{
-public:
-
-	FMLInferenceModelHlsl();
-	~FMLInferenceModelHlsl();
-
-	bool Init(TConstArrayView<uint8> ModelData);
-
-protected:
-
-	virtual void AddDispatchOps_RenderThread(FRDGBuilder& GraphBuilder) override;
-	virtual int PrepareTensorShapesAndData() override;
-
-private:
-
-	FMLOperatorHlsl* OpCreate(const FString& Name, TConstArrayView<FTensorDesc> InputTensorDesc, TConstArrayView<FTensorDesc> OutputTensorDescs, const UE::NNECore::FAttributeMap& Attributes);
-
-	TArray<FMLOperatorHlsl*>	Operators;
-};
-
-
-//
-//
-//
-FMLInferenceModelHlsl::FMLInferenceModelHlsl()
-{
-}
-
-//
-//
-//
-FMLInferenceModelHlsl::~FMLInferenceModelHlsl()
-{
-	// TODO: Release all the operators
-}
-
-//
-//
-//
-bool FMLInferenceModelHlsl::Init(TConstArrayView<uint8> ModelData)
-{
-	check(ModelData.Num() > 0);
-	FMLRuntimeFormat	Format;
-
-	if (!LoadModel(ModelData, Format))
-	{
-		return false;
-	}
-
-	// Loop over all operators in the model and create them
-	for (int32 Idx = 0; Idx < Format.Operators.Num(); ++Idx)
-	{
-		const FString TypeName = Format.Operators[Idx].TypeName;
-
-		TArray<FTensorDesc> Inputs;
-		TArray<FTensorDesc> Outputs;
-		UE::NNECore::FAttributeMap AttributeMap;
-
-		for (int32 InputTensorIndex : Format.Operators[Idx].InTensors)
-		{
-			Inputs.Emplace(AllSymbolicTensorDescs[InputTensorIndex]);
-		}
-		for (int32 OutputTensorIndex : Format.Operators[Idx].OutTensors)
-		{
-			Outputs.Emplace(AllSymbolicTensorDescs[OutputTensorIndex]);
-		}
-		for (const FMLFormatAttributeDesc& Desc : Format.Operators[Idx].Attributes)
-		{
-			AttributeMap.SetAttribute(Desc.Name, Desc.Value);
-		}
-
-		FMLOperatorHlsl* Op = OpCreate(TypeName, Inputs, Outputs, AttributeMap);
-
-		if (!Op) //Op.Shader.IsNull())
-		{
-			UE_LOG(LogNNX, Warning, TEXT("Failed to create operator:%s"), *TypeName);
-
-			// TODO: Cleanup operators
-			return false;
-		}
-
-		Operators.Add(Op);
-	}
-
-	return true;
-}
-
-//
-//
-//
-FMLOperatorHlsl* FMLInferenceModelHlsl::OpCreate(const FString& OpName, TConstArrayView<FTensorDesc> InputTensorDescs, TConstArrayView<FTensorDesc> OutputTensorDescs, const UE::NNECore::FAttributeMap& AttributeMap)
-{
-	FMLOperatorRegistryHlsl::OperatorCreateFunc CreateFn = FMLOperatorRegistryHlsl::Get()->OpFind(OpName);
-
-	if (!CreateFn)
-	{
-		UE_LOG(LogNNX, Warning, TEXT("Hlsl MLOperatorRegistry failed to find operator:%s"), *OpName);
-		return nullptr;
-	}
-
-	FMLOperatorHlsl* Op = CreateFn();
-
-	if (!Op->Initialize(InputTensorDescs, OutputTensorDescs, AttributeMap))
-	{
-		UE_LOG(LogNNX, Warning, TEXT("Hlsl engine: Error initializing operator:%s"), *OpName);
-		delete Op;
-		return nullptr;
-	}
-
-	return Op;
-}
-
-int FMLInferenceModelHlsl::PrepareTensorShapesAndData()
-{
-	check(AllTensorRDGs.Num() == AllSymbolicTensorDescs.Num());
-	
-	if (Operators.Num() == 0)
-	{
-		UE_LOG(LogNNX, Warning, TEXT("No operators in model"));
-		return -1;
-	}
-
-	// Run model preparation (including shape inference) on all operators
-	// This loop could be abstracted to a different engine/system as it apply on FTensorRef & IPrepareOperator witch are RDG agnostics.
-	static constexpr int32 MaxExpectedInput = 10;
-	TArray<FTensorRef, TInlineAllocator<MaxExpectedInput>> InputTensors;
-	TArray<FTensorRef> OutputTensors;
-	TArray<bool> AllInitializedTensors;
-
-	checkCode(
-		AllInitializedTensors.Init(false, AllSymbolicTensorDescs.Num());
-		for (int32 Idx : InputTensorIndices)
-		{
-			AllInitializedTensors[Idx] = true;
-		}
-		for (int32 Idx : WeightTensorIndices)
-		{
-			AllInitializedTensors[Idx] = true;
-		}
-	);
-
-	// Run model preparation (including shape inference) on all operators
-	// This loop could be abstracted to a different system as it apply on FTensorRef & IPrepareOperator witch are RDG agnostics.
-	for (int32 Idx = 0; Idx < Operators.Num(); ++Idx)
-	{
-		InputTensors.Empty();
-		OutputTensors.Empty();
-
-		//Operator inputs
-		for (int32 i : OperatorInputTensorIndices[Idx])
-		{
-			checkf(AllInitializedTensors[i] == true, TEXT("Input tensor %d for operator %d should have been initialized."), i, Idx);
-			InputTensors.Emplace(AllTensorRDGs[i]);
-		}
-		//Operator outputs
-		for (int32 i : OperatorOutputTensorIndices[Idx])
-		{
-			OutputTensors.Emplace(AllTensorRDGs[i]);
-			checkf(AllInitializedTensors[i] == false, TEXT("Output tensor %d for operator %d should not have been initialized yet."), i, Idx);
-			checkCode(AllInitializedTensors[i] = true);
-		}
-
-		const FMLOperatorHlsl* Op = Operators[Idx];
-
-		if (Op->PrepareOutputs(InputTensors, OutputTensors) != 0)
-		{
-			//Operator could not prepare the output tensors, meaning we can't allocate
-			//output buffer before running the model. This engine does not support this.
-			UE_LOG(LogNNX, Warning, TEXT("Could not deduce tensor shapes for this model during shape inference, HLSL engine wont support the model as it need to precompute all shapes for performance reasons."));
-			AllTensorRDGs.Empty();
-			return -1;
-		}
-	}
-
-	checkCode(
-		for (int i = 0; i < AllInitializedTensors.Num(); ++i)
-		{
-			checkf(AllInitializedTensors[i], TEXT("Tensor at index %d, was not initialized by model preparation."));
-		};
-	);
-
-	return 0;
-}
-
-//
-//
-//
-void FMLInferenceModelHlsl::AddDispatchOps_RenderThread(FRDGBuilder& GraphBuilder)
-{
-	static constexpr int32 MaxExpectedInput = 10;
-	TArray<FTensorRDGRef, TInlineAllocator<MaxExpectedInput>> InputTensors;
-
-	static constexpr int32 MaxExpectedOutput = 2;
-	TArray<FTensorRDGRef, TInlineAllocator<MaxExpectedOutput>> OutputTensors;
-
-	// Add passes for all operators
-	for (int32 Idx = 0; Idx < Operators.Num(); ++Idx)
-	{
-		InputTensors.Empty();
-		for (int32 i : OperatorInputTensorIndices[Idx])
-		{
-			InputTensors.Emplace(AllTensorRDGs[i]);
-		}
-		OutputTensors.Empty();
-		for (int32 i : OperatorOutputTensorIndices[Idx])
-		{
-			OutputTensors.Emplace(AllTensorRDGs[i]);
-		}
-
-		FMLOperatorHlsl* Op = Operators[Idx];
-
-		Op->Dispatch(GraphBuilder, InputTensors, OutputTensors);
-	}
-}
-
-//
-//
-//
-class FMLRuntimeHlsl : public FMLRuntimeRDG
+class FMLRuntimeHlsl : public NNX::FMLRuntimeRDG
 {
 public:
 
@@ -257,19 +39,19 @@ public:
 
 	bool Init()
 	{
-		FMLOperatorRegistryHlsl* registry = FMLOperatorRegistryHlsl::Get();
+		NNX::FMLOperatorRegistryHlsl* registry = NNX::FMLOperatorRegistryHlsl::Get();
 		check(registry != nullptr);
 
-		UE::NNIRuntimeRDG::Private::Hlsl::RegisterElementWiseUnaryOperators(*registry);
-		UE::NNIRuntimeRDG::Private::Hlsl::RegisterElementWiseBinaryOperators(*registry);
-		UE::NNIRuntimeRDG::Private::Hlsl::RegisterElementWiseVariadicOperators(*registry);
-		UE::NNIRuntimeRDG::Private::Hlsl::RegisterGemmOperator(*registry);
-		UE::NNIRuntimeRDG::Private::Hlsl::RegisterConvOperator(*registry);
-		UE::NNIRuntimeRDG::Private::Hlsl::RegisterConvTransposeOperator(*registry);
-		UE::NNIRuntimeRDG::Private::Hlsl::RegisterMatMulOperator(*registry);
-		UE::NNIRuntimeRDG::Private::Hlsl::RegisterInstanceNormalizationOperator(*registry);
-		UE::NNIRuntimeRDG::Private::Hlsl::RegisterUpsampleOperator(*registry);
-		UE::NNIRuntimeRDG::Private::Hlsl::RegisterPadOperator(*registry);
+		RegisterElementWiseUnaryOperators(*registry);
+		RegisterElementWiseBinaryOperators(*registry);
+		RegisterElementWiseVariadicOperators(*registry);
+		RegisterGemmOperator(*registry);
+		RegisterConvOperator(*registry);
+		RegisterConvTransposeOperator(*registry);
+		RegisterMatMulOperator(*registry);
+		RegisterInstanceNormalizationOperator(*registry);
+		RegisterUpsampleOperator(*registry);
+		RegisterPadOperator(*registry);
 
 		return true;
 	}
@@ -280,9 +62,9 @@ public:
 	}
 
 	// Returns flags from ERuntimeSupportFlags
-	virtual EMLRuntimeSupportFlags GetSupportFlags() const override
+	virtual NNX::EMLRuntimeSupportFlags GetSupportFlags() const override
 	{
-		return EMLRuntimeSupportFlags::RDG;
+		return NNX::EMLRuntimeSupportFlags::RDG;
 	}
 
 	virtual TArray<uint8> CreateModelData(FString FileType, TConstArrayView<uint8> FileData) override
@@ -292,8 +74,8 @@ public:
 			return {};
 		}
 
-		TUniquePtr<IModelOptimizer> Optimizer = CreateONNXToNNXModelOptimizer();
-		Optimizer->AddValidator(MakeShared<FModelValidatorHlsl>());
+		TUniquePtr<NNX::IModelOptimizer> Optimizer = NNX::CreateONNXToNNXModelOptimizer();
+		Optimizer->AddValidator(MakeShared<NNX::FModelValidatorHlsl>());
 
 		FNNIModelRaw InputModel;
 		InputModel.Data = FileData;
@@ -305,38 +87,41 @@ public:
 			return {};
 		}
 
-		return ConvertToModelData(OutputModel.Data);
+		return NNX::ConvertToModelData(OutputModel.Data);
 	};
 
-	virtual TUniquePtr<FMLInferenceModel> CreateModel(TConstArrayView<uint8> ModelData) override
+	virtual TUniquePtr<NNX::FMLInferenceModel> CreateModel(TConstArrayView<uint8> ModelData) override
 	{
 		if (!CanCreateModel(ModelData))
 		{
-			return TUniquePtr<FMLInferenceModel>();
+			return TUniquePtr<NNX::FMLInferenceModel>();
 		}
 
 		// Create the model and initialize it with the data not including the header
-		FMLInferenceModelHlsl* Model = new FMLInferenceModelHlsl();
+		FModel* Model = new FModel();
 		if (!Model->Init(ModelData))
 		{
 			delete Model;
-			return TUniquePtr<FMLInferenceModel>();
+			return TUniquePtr<NNX::FMLInferenceModel>();
 		}
-		return TUniquePtr<FMLInferenceModel>(Model);
+		return TUniquePtr<NNX::FMLInferenceModel>(Model);
 	}
 };
 
+} // namespace UE::NNIRuntimeRDG::Private::Hlsl
 
+namespace NNX
+{
 
 /** Globally accessible runtime */
-static TUniquePtr<FMLRuntimeHlsl> GHlslRuntime;
+static TUniquePtr<UE::NNIRuntimeRDG::Private::Hlsl::FMLRuntimeHlsl> GHlslRuntime;
 
 //
 //
 //
-static TUniquePtr<FMLRuntimeHlsl> FMLRuntimeHlslCreate()
+static TUniquePtr<UE::NNIRuntimeRDG::Private::Hlsl::FMLRuntimeHlsl> FMLRuntimeHlslCreate()
 {
-	auto Runtime = MakeUnique<FMLRuntimeHlsl>();
+	auto Runtime = MakeUnique<UE::NNIRuntimeRDG::Private::Hlsl::FMLRuntimeHlsl>();
 
 	if (!Runtime->Init())
 	{
@@ -371,4 +156,4 @@ void FMLRuntimeHlslShutdown()
 	}
 }
 
-} // NNX
+} // namespace NNX
diff --git a/Engine/Plugins/Experimental/NNE/Source/NNXRuntimeRDG/Private/NNXRuntimeRDG.h b/Engine/Plugins/Experimental/NNE/Source/NNXRuntimeRDG/Private/NNXRuntimeRDG.h
--- a/Engine/Plugins/Experimental/NNE/Source/NNXRuntimeRDG/Private/NNXRuntimeRDG.h
+++ b/Engine/Plugins/Experimental/NNE/Source/NNXRuntimeRDG/Private/NNXRuntimeRDG.h
@@ -56,7 +56,7 @@ struct IOperatorRDG
 
 class FTensorRDG : public FTensor
 {
-	FRDGBufferRef Buffer;
+	FRDGBufferRef Buffer{};
 	
 public:
 	static FTensorRDG Make(const FTensorDesc& TensorDesc, const FTensorShape& Shape, FRDGBufferRef Buffer)
@@ -73,6 +73,7 @@ public:
 		return TensorRDG;
 	}
 
+	bool HasBuffer() const { return Buffer != FRDGBufferRef{}; }
 	void SetBuffer(FRDGBufferRef Inbuffer){ Buffer = Inbuffer; }
 	FRDGBufferRef GetBuffer() const { return Buffer; }
 };
