NNE - NPP: Enable Input range scale and defining resolutions for the neural networks

diff --git a/Engine/Plugins/Experimental/NNE/Source/NNEHlslShaders/Internal/NNEHlslShadersNeuralPostProcessingCS.h b/Engine/Plugins/Experimental/NNE/Source/NNEHlslShaders/Internal/NNEHlslShadersNeuralPostProcessingCS.h
--- a/Engine/Plugins/Experimental/NNE/Source/NNEHlslShaders/Internal/NNEHlslShadersNeuralPostProcessingCS.h
+++ b/Engine/Plugins/Experimental/NNE/Source/NNEHlslShaders/Internal/NNEHlslShadersNeuralPostProcessingCS.h
@@ -62,6 +62,7 @@ namespace UE::NNEHlslShaders::Internal
 			SHADER_PARAMETER(int32, InputBufferWidth)
 			SHADER_PARAMETER(int32, InputBufferHeight)
 			SHADER_PARAMETER_RDG_BUFFER_UAV(RWBuffer<float>, InputBuffer)
+			SHADER_PARAMETER(float, RangeScale)
 		END_SHADER_PARAMETER_STRUCT()
 
 		static void ModifyCompilationEnvironment(const FGlobalShaderPermutationParameters& InParameters, FShaderCompilerEnvironment& OutEnvironment);
@@ -86,6 +87,7 @@ namespace UE::NNEHlslShaders::Internal
 			SHADER_PARAMETER(int32, InputTextureHeight)
 			SHADER_PARAMETER_RDG_BUFFER_UAV(RWBuffer<float>, AccumulationBuffer)
 			SHADER_PARAMETER(float, Weight)
+			SHADER_PARAMETER(float, RangeScale)
 		END_SHADER_PARAMETER_STRUCT()
 
 		static void ModifyCompilationEnvironment(const FGlobalShaderPermutationParameters& InParameters, FShaderCompilerEnvironment& OutEnvironment);
diff --git a/Engine/Plugins/Experimental/NNE/Source/NNEHlslShaders/Shaders/NNEHlslShadersNeuralPostProcessing.usf b/Engine/Plugins/Experimental/NNE/Source/NNEHlslShaders/Shaders/NNEHlslShadersNeuralPostProcessing.usf
--- a/Engine/Plugins/Experimental/NNE/Source/NNEHlslShaders/Shaders/NNEHlslShadersNeuralPostProcessing.usf
+++ b/Engine/Plugins/Experimental/NNE/Source/NNEHlslShaders/Shaders/NNEHlslShadersNeuralPostProcessing.usf
@@ -18,6 +18,7 @@ int OutputBufferHeight;
 RWBuffer<float> OutputBuffer;
 
 float Weight;
+float RangeScale;
 
 [numthreads(THREAD_GROUP_SIZE, THREAD_GROUP_SIZE, 1)]
 void ReadInput(in const uint3 DispatchThreadID : SV_DispatchThreadID)
@@ -72,9 +73,9 @@ void PreStep(in const uint3 DispatchThreadID : SV_DispatchThreadID)
 	int Idx = DispatchThreadID.y * InputBufferWidth + DispatchThreadID.x;
 	int Offset1 = InputBufferWidth * InputBufferHeight;
 	int Offset2 = Offset1 + Offset1;
-	InputBuffer[Idx] = Result.r;
-	InputBuffer[Idx + Offset1] = Result.g;
-	InputBuffer[Idx + Offset2] = Result.b;
+	InputBuffer[Idx] = RangeScale * Result.r;
+	InputBuffer[Idx + Offset1] = RangeScale * Result.g;
+	InputBuffer[Idx + Offset2] = RangeScale * Result.b;
 }
 
 [numthreads(THREAD_GROUP_SIZE, THREAD_GROUP_SIZE, 1)]
@@ -129,13 +130,13 @@ void PostStep(in const uint3 DispatchThreadID : SV_DispatchThreadID)
 
 	int ResultIdx = (DispatchThreadID.y * InputTextureWidth + DispatchThreadID.x) * 3;
 #if OVERWRITE == 0
-	AccumulationBuffer[ResultIdx + 0] += Weight * Result.r;
-	AccumulationBuffer[ResultIdx + 1] += Weight * Result.g;
-	AccumulationBuffer[ResultIdx + 2] += Weight * Result.b;
+	AccumulationBuffer[ResultIdx + 0] += Weight * Result.r / RangeScale;
+	AccumulationBuffer[ResultIdx + 1] += Weight * Result.g / RangeScale;
+	AccumulationBuffer[ResultIdx + 2] += Weight * Result.b / RangeScale;
 #else
-	AccumulationBuffer[ResultIdx + 0] = Weight * Result.r;
-	AccumulationBuffer[ResultIdx + 1] = Weight * Result.g;
-	AccumulationBuffer[ResultIdx + 2] = Weight * Result.b;
+	AccumulationBuffer[ResultIdx + 0] = Weight * Result.r / RangeScale;
+	AccumulationBuffer[ResultIdx + 1] = Weight * Result.g / RangeScale;
+	AccumulationBuffer[ResultIdx + 2] = Weight * Result.b / RangeScale;
 #endif
 }
 
diff --git a/Engine/Plugins/Experimental/NNE/Source/NNXRuntimeRDG/Private/NNERuntimeRDGNeuralPostProcessing.cpp b/Engine/Plugins/Experimental/NNE/Source/NNXRuntimeRDG/Private/NNERuntimeRDGNeuralPostProcessing.cpp
--- a/Engine/Plugins/Experimental/NNE/Source/NNXRuntimeRDG/Private/NNERuntimeRDGNeuralPostProcessing.cpp
+++ b/Engine/Plugins/Experimental/NNE/Source/NNXRuntimeRDG/Private/NNERuntimeRDGNeuralPostProcessing.cpp
@@ -64,6 +64,8 @@ int32 FNNENeuralPostProcessing::Add(FString RuntimeName, UNNEModelData* ModelDat
 	// Add the model to the map
 	Models.Add(LastId, Model);
 	SetWeight(LastId, 0.0);
+	SetRangeScale(LastId, 1.0);
+	SetInputSize(LastId, FIntPoint(-1, -1));
 
 	return LastId;
 }
@@ -77,6 +79,8 @@ bool FNNENeuralPostProcessing::Remove(int32 ModelId)
 	Disable(ModelId);
 	Models.Remove(ModelId);
 	Weights.Remove(ModelId);
+	RangeScales.Remove(ModelId);
+	InputSizes.Remove(ModelId);
 
 	return bResult;
 }
@@ -93,6 +97,30 @@ bool FNNENeuralPostProcessing::SetWeight(int32 ModelId, float Weight)
 	return false;
 }
 
+bool FNNENeuralPostProcessing::SetRangeScale(int32 ModelId, float RangeScale)
+{
+	FScopeLock Lock(&CriticalSection);
+
+	if (Models.Find(ModelId) != nullptr)
+	{
+		RangeScales.Add(ModelId, RangeScale);
+		return true;
+	}
+	return false;
+}
+
+bool FNNENeuralPostProcessing::SetInputSize(int32 ModelId, FIntPoint InputSize)
+{
+	FScopeLock Lock(&CriticalSection);
+
+	if (Models.Find(ModelId) != nullptr)
+	{
+		InputSizes.Add(ModelId, InputSize);
+		return true;
+	}
+	return false;
+}
+
 void FNNENeuralPostProcessing::Enable(int32 ModelId)
 {
 	FScopeLock Lock(&CriticalSection);
@@ -213,14 +241,26 @@ void FNNENeuralPostProcessing::PrePostProcessPass_RenderThread(FRDGBuilder& Grap
 				float OutputWeight = *Weights.Find(Pair.Key);
 				if (OutputWeight > WeightEpsilon)
 				{
+					float RangeScale = 1.0;
+					if (RangeScales.Find(Pair.Key))
+					{
+						RangeScale = *RangeScales.Find(Pair.Key);
+					}
+
+					FIntPoint InputSize = FIntPoint(-1, -1);
+					if (InputSizes.Find(Pair.Key))
+					{
+						InputSize = *InputSizes.Find(Pair.Key);
+					}
+
 					NNX::FSymbolicTensorShape InputShape = Pair.Value->GetInputTensorDescs()[0].GetShape();
 
 					checkf(InputShape.Rank() == 4, TEXT("Neural Post Processing requires models with input shape 1 x 3 x height x width!"))
 					checkf(InputShape.Data[0] == 1, TEXT("Neural Post Processing requires models with input shape 1 x 3 x height x width!"))
 					checkf(InputShape.Data[1] == 3, TEXT("Neural Post Processing requires models with input shape 1 x 3 x height x width!"))
 
-					int32 NeuralNetworkInputWidth = InputShape.Data[3] < 0 ? TextureSize.X : InputShape.Data[3];
-					int32 NeuralNetworkInputHeight = InputShape.Data[2] < 0 ? TextureSize.Y : InputShape.Data[2];
+					int32 NeuralNetworkInputWidth = InputShape.Data[3] >= 0 ? InputShape.Data[3] : (InputSize.X >= 0 ? InputSize.X : TextureSize.X);
+					int32 NeuralNetworkInputHeight = InputShape.Data[2] >= 0 ? InputShape.Data[2] : (InputSize.Y >= 0 ? InputSize.Y : TextureSize.Y);
 
 					FRDGBufferDesc InputBufferDesc = FRDGBufferDesc::CreateBufferDesc(sizeof(float), NeuralNetworkInputWidth * NeuralNetworkInputHeight * 3);
 					FRDGBufferRef InputBuffer = GraphBuilder.CreateBuffer(InputBufferDesc, *(FString("NNENeuralPostProcessing::NeuralNetowrkInput_") + FString::FromInt(Pair.Key)));
@@ -234,6 +274,7 @@ void FNNENeuralPostProcessing::PrePostProcessPass_RenderThread(FRDGBuilder& Grap
 					PreStepParameters->InputBufferWidth = NeuralNetworkInputWidth;
 					PreStepParameters->InputBufferHeight = NeuralNetworkInputHeight;
 					PreStepParameters->InputBuffer = InputBufferUAV;
+					PreStepParameters->RangeScale = RangeScale;
 
 					FIntVector PreStepThreadGroupCount = FIntVector(FMath::DivideAndRoundUp(NeuralNetworkInputWidth, FNeuralPostProcessingConstants::THREAD_GROUP_SIZE), FMath::DivideAndRoundUp(NeuralNetworkInputHeight, FNeuralPostProcessingConstants::THREAD_GROUP_SIZE), 1);
 					TShaderMapRef<TNeuralPostProcessingPreStepCS> PreStepShader(GlobalShaderMap);
@@ -282,6 +323,7 @@ void FNNENeuralPostProcessing::PrePostProcessPass_RenderThread(FRDGBuilder& Grap
 					PostStepParameters->InputTextureHeight = TextureSize.Y;
 					PostStepParameters->AccumulationBuffer = AccumulationBufferUAV;
 					PostStepParameters->Weight = OutputWeight;
+					PostStepParameters->RangeScale = RangeScale;
 
 					TNeuralPostProcessingPostStepCS::FPermutationDomain PermutationVector;
 					PermutationVector.Set<TNeuralPostProcessingPostStepCS::FNeuralPostProcessingOverwrite>(bOverwrite ? ENeuralPostProcessingOverwrite::Yes : ENeuralPostProcessingOverwrite::No);
@@ -354,6 +396,24 @@ bool UNNENeuralPostProcessing::SetWeight(int32 ModelId, float Weight)
 	return NeuralPostProcessing->SetWeight(ModelId, Weight);
 }
 
+bool UNNENeuralPostProcessing::SetRangeScale(int32 ModelId, float RangeScale)
+{
+	if (!NeuralPostProcessing.IsValid())
+	{
+		NeuralPostProcessing = FSceneViewExtensions::NewExtension<FNNENeuralPostProcessing>();
+	}
+	return NeuralPostProcessing->SetRangeScale(ModelId, RangeScale);
+}
+
+bool UNNENeuralPostProcessing::SetInputSize(int32 ModelId, FIntPoint InputSize)
+{
+	if (!NeuralPostProcessing.IsValid())
+	{
+		NeuralPostProcessing = FSceneViewExtensions::NewExtension<FNNENeuralPostProcessing>();
+	}
+	return NeuralPostProcessing->SetInputSize(ModelId, InputSize);
+}
+
 void UNNENeuralPostProcessing::Enable(int32 ModelId)
 {
 	if (!NeuralPostProcessing.IsValid())
diff --git a/Engine/Plugins/Experimental/NNE/Source/NNXRuntimeRDG/Public/NNERuntimeRDGNeuralPostProcessing.h b/Engine/Plugins/Experimental/NNE/Source/NNXRuntimeRDG/Public/NNERuntimeRDGNeuralPostProcessing.h
--- a/Engine/Plugins/Experimental/NNE/Source/NNXRuntimeRDG/Public/NNERuntimeRDGNeuralPostProcessing.h
+++ b/Engine/Plugins/Experimental/NNE/Source/NNXRuntimeRDG/Public/NNERuntimeRDGNeuralPostProcessing.h
@@ -21,6 +21,8 @@ public:
 	int32 Add(FString RuntimeName, UNNEModelData* ModelData);
 	bool Remove(int32 ModelId);
 	bool SetWeight(int32 ModelId, float Weight);
+	bool SetRangeScale(int32 ModelId, float RangeScale);
+	bool SetInputSize(int32 ModelId, FIntPoint InputSize);
 	void Enable(int32 ModelId);
 	void Disable(int32 ModelId);
 
@@ -40,6 +42,8 @@ private:
 	int32 LastId;
 	TMap<int32, TSharedPtr<NNX::FMLInferenceModel, ESPMode::ThreadSafe>> Models;
 	TMap<int32, float> Weights;
+	TMap<int32, float> RangeScales;
+	TMap<int32, FIntPoint> InputSizes;
 	TSet<int32> Enabled;
 
 };
@@ -60,6 +64,12 @@ public:
 	UFUNCTION(BlueprintCallable, Category = "NNE - Neural Network Engine")
 	bool SetWeight(int32 ModelId, float Weight);
 
+	UFUNCTION(BlueprintCallable, Category = "NNE - Neural Network Engine")
+	bool SetRangeScale(int32 ModelId, float RangeScale);
+	
+	UFUNCTION(BlueprintCallable, Category = "NNE - Neural Network Engine")
+	bool SetInputSize(int32 ModelId, FIntPoint InputSize);
+
 	UFUNCTION(BlueprintCallable, Category = "NNE - Neural Network Engine")
 	void Enable(int32 ModelId);
 
