Added more style transfer models (480p) and added toggle ("E" key).

diff --git a/Engine/Plugins/Experimental/NNE/Source/NNEHlslShaders/Shaders/NNEHlslShadersNeuralPostProcessing.usf b/Engine/Plugins/Experimental/NNE/Source/NNEHlslShaders/Shaders/NNEHlslShadersNeuralPostProcessing.usf
--- a/Engine/Plugins/Experimental/NNE/Source/NNEHlslShaders/Shaders/NNEHlslShadersNeuralPostProcessing.usf
+++ b/Engine/Plugins/Experimental/NNE/Source/NNEHlslShaders/Shaders/NNEHlslShadersNeuralPostProcessing.usf
@@ -31,9 +31,9 @@ void ReadInput(in const uint3 DispatchThreadID : SV_DispatchThreadID)
 	float4 InputColor = InputTexture.SampleLevel(InputTextureSampler, float2((float(DispatchThreadID.x) + 0.5) / float(InputTextureWidth), (float(DispatchThreadID.y) + 0.5) / float(InputTextureHeight)), 0);
 
 	int Idx = (DispatchThreadID.y * InputTextureWidth + DispatchThreadID.x) * 3;
-	AccumulationBuffer[Idx + 0] = Weight * InputColor.b;
+	AccumulationBuffer[Idx + 0] = Weight * InputColor.r;
 	AccumulationBuffer[Idx + 1] = Weight * InputColor.g;
-	AccumulationBuffer[Idx + 2] = Weight * InputColor.r;
+	AccumulationBuffer[Idx + 2] = Weight * InputColor.b;
 }
 
 [numthreads(THREAD_GROUP_SIZE, THREAD_GROUP_SIZE, 1)]
