NNE - Add CPUX runtime to sandbox.

diff --git a/Engine/Plugins/Experimental/NNE/Source/NNXQA/Resources/NNXQATestDesc.json b/Engine/Plugins/Experimental/NNE/Source/NNXQA/Resources/NNXQATestDesc.json
--- a/Engine/Plugins/Experimental/NNE/Source/NNXQA/Resources/NNXQATestDesc.json
+++ b/Engine/Plugins/Experimental/NNE/Source/NNXQA/Resources/NNXQATestDesc.json
@@ -41,6 +41,15 @@
 						{ "name": "NNXRuntimeDml",  "skip": true, "comment": "Model with multiple layer not yet implemented." }
 					]
 				},
+				{
+					"target": "NeuralMorphModel_global_Dax",
+					"tags": [ "SmokeTest" ],
+					"runtimes": [
+						{ "name": "NNXRuntimeORTDml", "skip": true, "comment": "" },
+						{ "name": "NNXRuntimeHlsl",	"skip": true, "comment": "Model with multiple layer not yet implemented."	},
+						{ "name": "NNXRuntimeDml",  "skip": true, "comment": "Model with multiple layer not yet implemented." }
+					]
+				},
 				{
 					"target": "NeuralMorphModel_Local",
 					"tags": [ "SmokeTest" ],
@@ -143,7 +152,7 @@
 						{ "name": "NNXRuntimeORTDml", "absolute_error":  1e-4, "relative_error": 0.03},
 						{ "name": "NNXRuntimeDml",    "absolute_error":  1e-4, "relative_error": 0.03},
 						{ "name": "NNXRuntimeHlsl",   "absolute_error":  1e-4, "relative_error": 0.03}
-					]},
+					], "additional_datasets": [ "Elementwise.Slow.OneInput" ]},
 				{	"target": "Erf", "tags": [ "SmokeTest" ], "runtimes": [
 						{ "name": "NNXRuntimeORTDml", "absolute_error":  1e-4, "relative_error": 0.03},
 						{ "name": "NNXRuntimeDml",    "absolute_error":  1e-4, "relative_error": 0.03},
