Fixed DirectML binding and uploading data when there are no weight tensors

#rnx

diff --git a/Engine/Plugins/Experimental/NNE/Source/NNXRuntimeRDG/Private/NNXRuntimeDML.cpp b/Engine/Plugins/Experimental/NNE/Source/NNXRuntimeRDG/Private/NNXRuntimeDML.cpp
--- a/Engine/Plugins/Experimental/NNE/Source/NNXRuntimeRDG/Private/NNXRuntimeDML.cpp
+++ b/Engine/Plugins/Experimental/NNE/Source/NNXRuntimeRDG/Private/NNXRuntimeDML.cpp
@@ -1194,12 +1194,17 @@ public:
 		
 		BindingTable->BindInputs(1, &InputBindArrayDesc);
 		
+		DML_BUFFER_BINDING			PersistBind {};
+		DML_BINDING_DESC			PersistBindDesc{ DML_BINDING_TYPE_BUFFER, &PersistBind };
+
+		if (PersistBuff)
+		{
 #ifdef NNE_USE_D3D12_RESOURCES
-		DML_BUFFER_BINDING			PersistBind = { PersistBuff, 0, PersistBuff->GetDesc().Width };
+			PersistBind = DML_BUFFER_BINDING { PersistBuff, 0, PersistBuff->GetDesc().Width };
 #else
-		DML_BUFFER_BINDING			PersistBind = MakeBind(PersistBuff);
+			PersistBind = MakeBind(PersistBuff);
 #endif
-		DML_BINDING_DESC			PersistBindDesc{ DML_BINDING_TYPE_BUFFER, &PersistBind };
+		}
 
 		BindingTable->BindOutputs(1, &PersistBindDesc);
 
@@ -1971,37 +1976,42 @@ bool FMLInferenceModelDml::InitCompiledOp(TConstArrayView<int32> OpInputIndices,
 				Inputs.Emplace(nullptr);
 			}
 
-			FGPUFenceRHIRef UploadFence = RHICmdList.CreateGPUFence(TEXT("FMLInferenceModel_UploadFence"));
-			FRHIBuffer*		UploadBuff = CreateRHIBuffer(RHICmdList, TensorDataSize, BUF_ShaderResource | BUF_Dynamic | BUF_FastVRAM, ERHIAccess::CopySrc, TEXT("FMLInferenceModel_UploadBuffer"));
-			uint8*			UploadBuffPtr = static_cast<uint8*>(RHICmdList.LockBuffer(UploadBuff, 0, TensorDataSize, RLM_WriteOnly_NoOverwrite));
-			uint64			UploadOffset = 0;
-
 			TArray<CD3DX12_RESOURCE_BARRIER, TInlineAllocator<MaxNumInputs>>	Barriers;
+			FGPUFenceRHIRef	UploadFence = nullptr;
 
-			for (const FTensorRDG& Tensor : WeightTensorRDGs)
+			if (TensorDataSize)
 			{
-				TConstArrayView<uint8>	TensorData = Tensor.GetPreparedData<uint8>();
+				UploadFence = RHICmdList.CreateGPUFence(TEXT("FMLInferenceModel_UploadFence"));
+				
+				FRHIBuffer*		UploadBuff = CreateRHIBuffer(RHICmdList, TensorDataSize, BUF_ShaderResource | BUF_Dynamic | BUF_FastVRAM, ERHIAccess::CopySrc, TEXT("FMLInferenceModel_UploadBuffer"));
+				uint8*			UploadBuffPtr = static_cast<uint8*>(RHICmdList.LockBuffer(UploadBuff, 0, TensorDataSize, RLM_WriteOnly_NoOverwrite));
+				uint64			UploadOffset = 0;
 
-				FBufferRHIRef WeightBuff;
+				for (const FTensorRDG& Tensor : WeightTensorRDGs)
+				{
+					TConstArrayView<uint8>	TensorData = Tensor.GetPreparedData<uint8>();
 
-				WeightBuff = CreateRHIBuffer(RHICmdList, TensorData.Num(), WeightBuffUsage, WeightBuffAccess, TEXT("FMLInferenceModelDml_TensorWeights"));
+					FBufferRHIRef WeightBuff;
+
+					WeightBuff = CreateRHIBuffer(RHICmdList, TensorData.Num(), WeightBuffUsage, WeightBuffAccess, TEXT("FMLInferenceModelDml_TensorWeights"));
 				
-				FMemory::Memcpy(UploadBuffPtr + UploadOffset, TensorData.GetData(), TensorData.Num());
-				RHICmdList.CopyBufferRegion(WeightBuff, 0, UploadBuff, UploadOffset, TensorData.Num());
-				UploadOffset += Align(TensorData.Num(), DML_MINIMUM_BUFFER_TENSOR_ALIGNMENT);
+					FMemory::Memcpy(UploadBuffPtr + UploadOffset, TensorData.GetData(), TensorData.Num());
+					RHICmdList.CopyBufferRegion(WeightBuff, 0, UploadBuff, UploadOffset, TensorData.Num());
+					UploadOffset += Align(TensorData.Num(), DML_MINIMUM_BUFFER_TENSOR_ALIGNMENT);
 
-				Inputs.Emplace(WeightBuff);
+					Inputs.Emplace(WeightBuff);
 				
-				Barriers.Emplace(
-					CD3DX12_RESOURCE_BARRIER::Transition(
-						DynamicRHI->RHIGetResource(WeightBuff),
-						D3D12_RESOURCE_STATE_COPY_DEST,
-						D3D12_RESOURCE_STATE_UNORDERED_ACCESS)
-				);
-			}
+					Barriers.Emplace(
+						CD3DX12_RESOURCE_BARRIER::Transition(
+							DynamicRHI->RHIGetResource(WeightBuff),
+							D3D12_RESOURCE_STATE_COPY_DEST,
+							D3D12_RESOURCE_STATE_UNORDERED_ACCESS)
+					);
+				}
 
-			RHICmdList.UnlockBuffer(UploadBuff);
-			RHICmdList.WriteGPUFence(UploadFence);
+				RHICmdList.UnlockBuffer(UploadBuff);
+				RHICmdList.WriteGPUFence(UploadFence);
+			}
 
 			if (MemSizePersist)
 			{
@@ -2039,7 +2049,7 @@ bool FMLInferenceModelDml::InitCompiledOp(TConstArrayView<int32> OpInputIndices,
 			RHICmdList.EnqueueLambda(
 				[this, Inputs, Barriers, InitTempBuff, UploadFence](FRHICommandListImmediate& RHICmdList)
 				{
-					while (UploadFence->NumPendingWriteCommands.GetValue() > 0)
+					while (UploadFence && UploadFence->NumPendingWriteCommands.GetValue() > 0)
 					{
 						FPlatformProcess::Sleep(0.001);
 					}
@@ -2051,7 +2061,11 @@ bool FMLInferenceModelDml::InitCompiledOp(TConstArrayView<int32> OpInputIndices,
 					BindingTable->Bind(OpInit, Inputs, PersistBuff, InitTempBuff);
 			
 					D3DCmdList = DynamicRHI->RHIGetGraphicsCommandList(DevCtx->DeviceIndex);
-					D3DCmdList->ResourceBarrier(Barriers.Num(), Barriers.GetData());
+
+					if (!Barriers.IsEmpty())
+					{
+						D3DCmdList->ResourceBarrier(Barriers.Num(), Barriers.GetData());
+					}
 					D3DCmdList->SetDescriptorHeaps(1, &DescHeap);
 					DevCtx->CmdRec->RecordDispatch(D3DCmdList, OpInit, BindingTable->Get());
 
