Fixed DirectML runtime abort in debug build when cooking

#rnx

diff --git a/Engine/Plugins/Experimental/NNE/Source/NNXRuntimeRDG/Private/NNXRuntimeDML.cpp b/Engine/Plugins/Experimental/NNE/Source/NNXRuntimeRDG/Private/NNXRuntimeDML.cpp
--- a/Engine/Plugins/Experimental/NNE/Source/NNXRuntimeRDG/Private/NNXRuntimeDML.cpp
+++ b/Engine/Plugins/Experimental/NNE/Source/NNXRuntimeRDG/Private/NNXRuntimeDML.cpp
@@ -887,7 +887,7 @@ public:
 		return TUniquePtr<FMLInferenceModel>(Model);
 	}
 
-	bool Init();
+	bool Init(bool bRegisterOnlyOperators);
 
 private:
 
@@ -913,12 +913,18 @@ FMLRuntimeDml::~FMLRuntimeDml()
 //
 //
 //
-bool FMLRuntimeDml::Init()
+bool FMLRuntimeDml::Init(bool bRegisterOnlyOperators)
 {
 	RegisterElementWiseUnaryOperators();
 	RegisterElementWiseBinaryOperators();
 	RegisterGemmOperator();
 
+	if (bRegisterOnlyOperators)
+	{
+		UE_LOG(LogNNX, Display, TEXT("Registering only operators"));
+		return true;
+	}
+
 	HRESULT Res;
 
 	// In order to use DirectML we need D3D12
@@ -2275,11 +2281,11 @@ int FMLInferenceModelDml::PrepareTensorShapesAndData()
 //
 //
 //
-static TUniquePtr<FMLRuntimeDml> FDmlRuntimeCreate()
+static TUniquePtr<FMLRuntimeDml> FDmlRuntimeCreate(bool bRegisterOnlyOperators)
 {
 	auto Runtime = MakeUnique<FMLRuntimeDml>();
 
-	if (!Runtime->Init())
+	if (!Runtime->Init(bRegisterOnlyOperators))
 	{
 		UE_LOG(LogNNX, Warning, TEXT("Failed to create NNX DML runtime"));
 		Runtime.Release();
@@ -2295,41 +2301,53 @@ IRuntime* FMLRuntimeDmlStartup()
 {
 	if (!GDmlRuntime)
 	{
-#ifdef DIRECTML_BIN_PATH
+		bool bIsD3D12RHI = GDynamicRHI && GDynamicRHI->GetInterfaceType() == ERHIInterfaceType::D3D12;		
+		bool bLoadDirectML = true;
 
-		ID3D12DynamicRHI* DynamicRHI = GetID3D12PlatformDynamicRHI();
+		if (IsRunningCommandlet() && !IsAllowCommandletRendering())
+		{
+			UE_LOG(LogNNX, Display, TEXT("Running inside commandlet without rendering"));
+			bLoadDirectML = false;
+		}
 
-		const FString DirectMLRuntimeBinPath = TEXT(PREPROCESSOR_TO_STRING(DIRECTML_BIN_PATH));
-		FString DirectMLDLLPaths[2];
-		int32	NumPaths = 1;
+#ifdef DIRECTML_BIN_PATH
 		
-		DirectMLDLLPaths[0] = DirectMLRuntimeBinPath / TEXT("DirectML.dll");
-
-		if (DynamicRHI->IsD3DDebugEnabled())
+		if (bIsD3D12RHI && bLoadDirectML)
 		{
-			DirectMLDLLPaths[1] = DirectMLRuntimeBinPath / TEXT("DirectML.Debug.dll");
-			++NumPaths;
-		}
+			const FString DirectMLRuntimeBinPath = TEXT(PREPROCESSOR_TO_STRING(DIRECTML_BIN_PATH));
+			FString DirectMLDLLPaths[2];
+			int32	NumPaths = 1;
+		
+			DirectMLDLLPaths[0] = DirectMLRuntimeBinPath / TEXT("DirectML.dll");
 
-		FPlatformProcess::PushDllDirectory(*DirectMLRuntimeBinPath);
+			if (GetID3D12PlatformDynamicRHI()->IsD3DDebugEnabled())
+			{
+				DirectMLDLLPaths[1] = DirectMLRuntimeBinPath / TEXT("DirectML.Debug.dll");
+				++NumPaths;
+			}
 
-		for (int32 Idx = 0; Idx < NumPaths; ++Idx)
-		{
-			if (!FPaths::FileExists(DirectMLDLLPaths[Idx]))
+			FPlatformProcess::PushDllDirectory(*DirectMLRuntimeBinPath);
+
+			for (int32 Idx = 0; Idx < NumPaths; ++Idx)
 			{
-				const FString ErrorMessage = FString::Format(TEXT("DirectML DLL file not found in \"{0}\"."),
-					{ IFileManager::Get().ConvertToAbsolutePathForExternalAppForRead(*DirectMLDLLPaths[Idx])});
-				UE_LOG(LogNNX, Warning, TEXT("NNXRuntimeDll:%s"), *ErrorMessage);
-				checkf(false, TEXT("%s"), *ErrorMessage);
+				if (!FPaths::FileExists(DirectMLDLLPaths[Idx]))
+				{
+					const FString ErrorMessage = FString::Format(TEXT("DirectML DLL file not found in \"{0}\"."),
+						{ IFileManager::Get().ConvertToAbsolutePathForExternalAppForRead(*DirectMLDLLPaths[Idx])});
+					UE_LOG(LogNNX, Warning, TEXT("NNXRuntimeDll:%s"), *ErrorMessage);
+					checkf(false, TEXT("%s"), *ErrorMessage);
+				}
+
+				FPlatformProcess::GetDllHandle(*DirectMLDLLPaths[Idx]);
 			}
 
-			FPlatformProcess::GetDllHandle(*DirectMLDLLPaths[Idx]);
+			FPlatformProcess::PopDllDirectory(*DirectMLRuntimeBinPath);
 		}
-
-		FPlatformProcess::PopDllDirectory(*DirectMLRuntimeBinPath);
 #endif
 
-		GDmlRuntime = FDmlRuntimeCreate();
+		const bool bRegisterOnlyOperators = !bLoadDirectML;
+
+		GDmlRuntime = FDmlRuntimeCreate(bRegisterOnlyOperators);
 	}
 
 	return GDmlRuntime.Get();
