Minimal API - NNECore/NNETypes

Problem:
* NNE API surface is large atm potentially creating a support burden and harming adoption for 5.2.

Solution:
* reduce the API surface as discussed in design discutions. This first CL is targeting NNECore/NNETypes as a first step.

Test:
* All tests green under win64 editor (with CL 23256970 witch disabled some failing tests on main to be investigated on there own) no regression without it compared to main

diff --git a/Engine/Plugins/Experimental/NNE/Source/NNEHlslShaders/Internal/NNEHlslShadersConvTransposeCS.h b/Engine/Plugins/Experimental/NNE/Source/NNEHlslShaders/Internal/NNEHlslShadersConvTransposeCS.h
--- a/Engine/Plugins/Experimental/NNE/Source/NNEHlslShaders/Internal/NNEHlslShadersConvTransposeCS.h
+++ b/Engine/Plugins/Experimental/NNE/Source/NNEHlslShaders/Internal/NNEHlslShadersConvTransposeCS.h
@@ -5,7 +5,6 @@
 #include "GlobalShader.h"
 #include "ShaderParameterUtils.h"
 #include "RenderGraphUtils.h"
-#include "NNXTypes.h"
 
 namespace UE::NNEHlslShaders::Internal
 {
diff --git a/Engine/Plugins/Experimental/NNE/Source/NNEHlslShaders/Internal/NNEHlslShadersGatherCS.h b/Engine/Plugins/Experimental/NNE/Source/NNEHlslShaders/Internal/NNEHlslShadersGatherCS.h
--- a/Engine/Plugins/Experimental/NNE/Source/NNEHlslShaders/Internal/NNEHlslShadersGatherCS.h
+++ b/Engine/Plugins/Experimental/NNE/Source/NNEHlslShaders/Internal/NNEHlslShadersGatherCS.h
@@ -5,7 +5,8 @@
 #include "GlobalShader.h"
 #include "ShaderParameterUtils.h"
 #include "RenderGraphUtils.h"
-#include "NNXTypes.h"
+
+namespace UE::NNECore::Internal { class FTensor; }
 
 namespace UE::NNEHlslShaders::Internal
 {
@@ -42,7 +43,7 @@ namespace UE::NNEHlslShaders::Internal
 
 		static void ModifyCompilationEnvironment(const FGlobalShaderPermutationParameters& InParameters, FShaderCompilerEnvironment& OutEnvironment);
 
-		static void FillInParameters(int32 Axis, const NNX::FTensor& Data, const NNX::FTensor& Indices, FParameters& Parameters);
+		static void FillInParameters(int32 Axis, const NNECore::Internal::FTensor& Data, const NNECore::Internal::FTensor& Indices, FParameters& Parameters);
 
 		static FIntVector GetGroupCount(const FParameters& Parameters);
 	};
diff --git a/Engine/Plugins/Experimental/NNE/Source/NNEHlslShaders/Internal/NNEHlslShadersGemmCS.h b/Engine/Plugins/Experimental/NNE/Source/NNEHlslShaders/Internal/NNEHlslShadersGemmCS.h
--- a/Engine/Plugins/Experimental/NNE/Source/NNEHlslShaders/Internal/NNEHlslShadersGemmCS.h
+++ b/Engine/Plugins/Experimental/NNE/Source/NNEHlslShaders/Internal/NNEHlslShadersGemmCS.h
@@ -5,7 +5,8 @@
 #include "GlobalShader.h"
 #include "ShaderParameterUtils.h"
 #include "RenderGraphUtils.h"
-#include "NNXTypes.h"
+
+namespace UE::NNECore::Internal { class FTensor; }
 
 namespace UE::NNEHlslShaders::Internal
 {
@@ -79,9 +80,9 @@ namespace UE::NNEHlslShaders::Internal
 
 		static void ModifyCompilationEnvironment(const FGlobalShaderPermutationParameters& InParameters, FShaderCompilerEnvironment& OutEnvironment);
 
-		static void FillInParameters(float Alpha, float Beta, int32 TransA, int32 TransB, const NNX::FTensor& InputA, const NNX::FTensor& InputB, const NNX::FTensor* InputC, float CScalar, FParameters& Parameters);
+		static void FillInParameters(float Alpha, float Beta, int32 TransA, int32 TransB, const NNECore::Internal::FTensor& InputA, const NNECore::Internal::FTensor& InputB, const NNECore::Internal::FTensor* InputC, float CScalar, FParameters& Parameters);
 
-		static void FillInParametersMatMul(const NNX::FTensor& InputA, const NNX::FTensor& InputB, FParameters& Parameters);
+		static void FillInParametersMatMul(const NNECore::Internal::FTensor& InputA, const NNECore::Internal::FTensor& InputB, FParameters& Parameters);
 
 		static FIntVector GetGroupCount(const FParameters& Parameters, EGemmAlgorithm Algorithm, int32 NumStackDimensions);
 		static EGemmAlgorithm GetAlgorithm(const FParameters& Parameters);
diff --git a/Engine/Plugins/Experimental/NNE/Source/NNEHlslShaders/Internal/NNEHlslShadersInstanceNormalizationCS.h b/Engine/Plugins/Experimental/NNE/Source/NNEHlslShaders/Internal/NNEHlslShadersInstanceNormalizationCS.h
--- a/Engine/Plugins/Experimental/NNE/Source/NNEHlslShaders/Internal/NNEHlslShadersInstanceNormalizationCS.h
+++ b/Engine/Plugins/Experimental/NNE/Source/NNEHlslShaders/Internal/NNEHlslShadersInstanceNormalizationCS.h
@@ -6,7 +6,7 @@
 #include "ShaderParameterUtils.h"
 #include "RenderGraphUtils.h"
 
-namespace NNX { class FTensor; }
+namespace UE::NNECore::Internal { class FTensor; }
 
 namespace UE::NNEHlslShaders::Internal
 {
@@ -40,7 +40,7 @@ namespace UE::NNEHlslShaders::Internal
 			SHADER_PARAMETER_RDG_BUFFER_UAV(RWBuffer<float>, Output)
 		END_SHADER_PARAMETER_STRUCT()
 
-		static void FillInParameters(float Epsilon, const NNX::FTensor& Input, FParameters& Parameters);
+		static void FillInParameters(float Epsilon, const NNECore::Internal::FTensor& Input, FParameters& Parameters);
 
 		static FIntVector GetGroupCount(const FParameters& Parameters, EInstanceNormalizationAlgorithm Algorithm);
 		static EInstanceNormalizationAlgorithm GetAlgorithm(const FParameters& Parameters);
diff --git a/Engine/Plugins/Experimental/NNE/Source/NNEHlslShaders/Private/NNEHlslShadersGatherCS.cpp b/Engine/Plugins/Experimental/NNE/Source/NNEHlslShaders/Private/NNEHlslShadersGatherCS.cpp
--- a/Engine/Plugins/Experimental/NNE/Source/NNEHlslShaders/Private/NNEHlslShadersGatherCS.cpp
+++ b/Engine/Plugins/Experimental/NNE/Source/NNEHlslShaders/Private/NNEHlslShadersGatherCS.cpp
@@ -1,6 +1,7 @@
 // Copyright Epic Games, Inc. All Rights Reserved.
 
 #include "NNEHlslShadersGatherCS.h"
+#include "NNECoreTensor.h"
 
 namespace UE::NNEHlslShaders::Internal
 {
@@ -15,7 +16,7 @@ namespace UE::NNEHlslShaders::Internal
 	}
 
 	//template <typename DataElementType, typename IndicesElementType>
-	void TGatherCS::FillInParameters(int32 Axis, const NNX::FTensor& Data, const NNX::FTensor& Indices, FParameters& Parameters)
+	void TGatherCS::FillInParameters(int32 Axis, const NNECore::Internal::FTensor& Data, const NNECore::Internal::FTensor& Indices, FParameters& Parameters)
 	{
 		Parameters.Axis = Axis;
 
@@ -25,17 +26,17 @@ namespace UE::NNEHlslShaders::Internal
 		Parameters.OutputSize = 1;
 		for (int32 i = 0; i < Axis; i++, index++)
 		{
-			OutputShape[index] = Data.GetShape().Data[i];
+			OutputShape[index] = Data.GetShape().GetData()[i];
 			Parameters.OutputSize *= OutputShape[index];
 		}
 		for (int32 i = 0; i < Indices.GetShape().Rank(); i++, index++)
 		{
-			OutputShape[index] = Indices.GetShape().Data[i];
+			OutputShape[index] = Indices.GetShape().GetData()[i];
 			Parameters.OutputSize *= OutputShape[index];
 		}
 		for (int32 i = (Axis + 1); i < Data.GetShape().Rank(); i++, index++)
 		{
-			OutputShape[index] = Data.GetShape().Data[i];
+			OutputShape[index] = Data.GetShape().GetData()[i];
 			Parameters.OutputSize *= OutputShape[index];
 		}
 
@@ -47,7 +48,7 @@ namespace UE::NNEHlslShaders::Internal
 		{
 			Parameters.DataStride_IndicesStride_OutputStride[i].X = DataStride;
 			Parameters.OneDivDataStride_OneDivIndicesStride_OneDivOutputStride[i].X = 1.0 / (float)DataStride;
-			DataStride *= Data.GetShape().Data[i];
+			DataStride *= Data.GetShape().GetData()[i];
 		}
 
 		int32 IndicesStride = 1;
@@ -55,7 +56,7 @@ namespace UE::NNEHlslShaders::Internal
 		{
 			Parameters.DataStride_IndicesStride_OutputStride[i].Y = IndicesStride;
 			Parameters.OneDivDataStride_OneDivIndicesStride_OneDivOutputStride[i].Y = 1.0 / (float)IndicesStride;
-			IndicesStride *= Indices.GetShape().Data[i];
+			IndicesStride *= Indices.GetShape().GetData()[i];
 		}
 
 		int32 OutputStride = 1;
diff --git a/Engine/Plugins/Experimental/NNE/Source/NNEHlslShaders/Private/NNEHlslShadersGemmCS.cpp b/Engine/Plugins/Experimental/NNE/Source/NNEHlslShaders/Private/NNEHlslShadersGemmCS.cpp
--- a/Engine/Plugins/Experimental/NNE/Source/NNEHlslShaders/Private/NNEHlslShadersGemmCS.cpp
+++ b/Engine/Plugins/Experimental/NNE/Source/NNEHlslShaders/Private/NNEHlslShadersGemmCS.cpp
@@ -2,6 +2,7 @@
 
 #include "NNEHlslShadersGemmCS.h"
 #include "NNXCore.h"
+#include "NNECoreTensor.h"
 #include "Algo/Accumulate.h"
 
 namespace UE::NNEHlslShaders::Internal
@@ -17,7 +18,7 @@ namespace UE::NNEHlslShaders::Internal
 		TArray<uint32> StackStrideB;
 	};
 
-	static GemmMatrixParameters GetMatrixParametersMatMul(const NNX::FTensor& InputA, const NNX::FTensor& InputB)
+	static GemmMatrixParameters GetMatrixParametersMatMul(const NNECore::Internal::FTensor& InputA, const NNECore::Internal::FTensor& InputB)
 	{
 		check(InputA.GetShape().Rank() != 0);
 		check(InputB.GetShape().Rank() != 0);
@@ -35,8 +36,8 @@ namespace UE::NNEHlslShaders::Internal
 			const int32 NumDimToCheck = FMath::Min(NumStackDimensionsA, NumStackDimensionsB);
 			for (int32 i = 0; i < NumDimToCheck; i++)
 			{
-				const uint32 VolumeA = InputA.GetShape().Data[NumStackDimensionsA - 1 - i];
-				const uint32 VolumeB = InputB.GetShape().Data[NumStackDimensionsB - 1 - i];
+				const uint32 VolumeA = InputA.GetShape().GetData()[NumStackDimensionsA - 1 - i];
+				const uint32 VolumeB = InputB.GetShape().GetData()[NumStackDimensionsB - 1 - i];
 
 				check(VolumeA == 1 || VolumeB == 1 || VolumeA == VolumeB);
 			}
@@ -45,10 +46,10 @@ namespace UE::NNEHlslShaders::Internal
 		const bool IsVectorA = InputA.GetShape().Rank() == 1;
 		const bool IsVectorB = InputB.GetShape().Rank() == 1;
 
-		Result.M = IsVectorA ? 1 : InputA.GetShape().Data[InputA.GetShape().Rank() - 2];
-		Result.N = IsVectorB ? 1 : InputB.GetShape().Data[InputB.GetShape().Rank() - 1];
-		Result.K = InputA.GetShape().Data[IsVectorA ? 0 : InputA.GetShape().Rank() - 1];
-		check(InputB.GetShape().Data[IsVectorB ? 0 : InputB.GetShape().Rank() - 2] == Result.K);
+		Result.M = IsVectorA ? 1 : InputA.GetShape().GetData()[InputA.GetShape().Rank() - 2];
+		Result.N = IsVectorB ? 1 : InputB.GetShape().GetData()[InputB.GetShape().Rank() - 1];
+		Result.K = InputA.GetShape().GetData()[IsVectorA ? 0 : InputA.GetShape().Rank() - 1];
+		check(InputB.GetShape().GetData()[IsVectorB ? 0 : InputB.GetShape().Rank() - 2] == Result.K);
 
 		Result.StackShapeA.Init(1, NumStackDimensions);
 		Result.StackShapeB.Init(1, NumStackDimensions);
@@ -58,8 +59,8 @@ namespace UE::NNEHlslShaders::Internal
 			int32 IdxA = InputA.GetShape().Rank() - 3 - i;
 			int32 IdxB = InputB.GetShape().Rank() - 3 - i;
 
-			Result.StackShapeA[Result.StackShapeA.Num() - 1 - i] = IdxA >= 0 ? InputA.GetShape().Data[IdxA] : 1;
-			Result.StackShapeB[Result.StackShapeB.Num() - 1 - i] = IdxB >= 0 ? InputB.GetShape().Data[IdxB] : 1;
+			Result.StackShapeA[Result.StackShapeA.Num() - 1 - i] = IdxA >= 0 ? InputA.GetShape().GetData()[IdxA] : 1;
+			Result.StackShapeB[Result.StackShapeB.Num() - 1 - i] = IdxB >= 0 ? InputB.GetShape().GetData()[IdxB] : 1;
 		}
 
 		Result.StackStrideA.Init(1, NumStackDimensions);
@@ -80,15 +81,15 @@ namespace UE::NNEHlslShaders::Internal
 		OutEnvironment.SetDefine(TEXT("MAX_NUM_STACK_DIMENSIONS"), FGemmConstants::MAX_NUM_STACK_DIMENSIONS);
 	}
 
-	void TGemmCS::FillInParameters(float Alpha, float Beta, int32 TransA, int32 TransB, const NNX::FTensor& InputA, const NNX::FTensor& InputB, const NNX::FTensor* InputC, float CScalar, TGemmCS::FParameters& Parameters)
+	void TGemmCS::FillInParameters(float Alpha, float Beta, int32 TransA, int32 TransB, const NNECore::Internal::FTensor& InputA, const NNECore::Internal::FTensor& InputB, const NNECore::Internal::FTensor* InputC, float CScalar, TGemmCS::FParameters& Parameters)
 	{
 		check(InputA.GetShape().Rank() == 2);
 		check(InputB.GetShape().Rank() == 2);
 
-		uint32 M = TransA != 0 ? InputA.GetShape().Data[1] : InputA.GetShape().Data[0];
-		uint32 K = TransA != 0 ? InputA.GetShape().Data[0] : InputA.GetShape().Data[1];
-		uint32 N = TransB != 0 ? InputB.GetShape().Data[0] : InputB.GetShape().Data[1];
-		check(K == (TransB != 0 ? InputB.GetShape().Data[1] : InputB.GetShape().Data[0]));
+		uint32 M = TransA != 0 ? InputA.GetShape().GetData()[1] : InputA.GetShape().GetData()[0];
+		uint32 K = TransA != 0 ? InputA.GetShape().GetData()[0] : InputA.GetShape().GetData()[1];
+		uint32 N = TransB != 0 ? InputB.GetShape().GetData()[0] : InputB.GetShape().GetData()[1];
+		check(K == (TransB != 0 ? InputB.GetShape().GetData()[1] : InputB.GetShape().GetData()[0]));
 
 		Parameters.Alpha = Alpha;
 		Parameters.Beta = Beta;
@@ -104,14 +105,14 @@ namespace UE::NNEHlslShaders::Internal
 		Parameters.CHeight = 0;
 		if (InputC != nullptr)
 		{
-			Parameters.CWidth = InputC->GetShape().Rank() == 0 ? 0 : InputC->GetShape().Data[InputC->GetShape().Rank() - 1];;
-			Parameters.CHeight = InputC->GetShape().Rank() == 0 ? 0 : InputC->GetShape().Rank() == 1 ? 1 : InputC->GetShape().Data[InputC->GetShape().Rank() - 2];
+			Parameters.CWidth = InputC->GetShape().Rank() == 0 ? 0 : InputC->GetShape().GetData()[InputC->GetShape().Rank() - 1];;
+			Parameters.CHeight = InputC->GetShape().Rank() == 0 ? 0 : InputC->GetShape().Rank() == 1 ? 1 : InputC->GetShape().GetData()[InputC->GetShape().Rank() - 2];
 		}
 		
 		Parameters.CScalar = CScalar;
 	}
 
-	void TGemmCS::FillInParametersMatMul(const NNX::FTensor& InputA, const NNX::FTensor& InputB, TGemmCS::FParameters& Parameters)
+	void TGemmCS::FillInParametersMatMul(const NNECore::Internal::FTensor& InputA, const NNECore::Internal::FTensor& InputB, TGemmCS::FParameters& Parameters)
 	{
 		GemmMatrixParameters Params = GetMatrixParametersMatMul(InputA, InputB);
 
diff --git a/Engine/Plugins/Experimental/NNE/Source/NNEHlslShaders/Private/NNEHlslShadersInstanceNormalizationCS.cpp b/Engine/Plugins/Experimental/NNE/Source/NNEHlslShaders/Private/NNEHlslShadersInstanceNormalizationCS.cpp
--- a/Engine/Plugins/Experimental/NNE/Source/NNEHlslShaders/Private/NNEHlslShadersInstanceNormalizationCS.cpp
+++ b/Engine/Plugins/Experimental/NNE/Source/NNEHlslShaders/Private/NNEHlslShadersInstanceNormalizationCS.cpp
@@ -2,7 +2,7 @@
 
 #include "NNEHlslShadersInstanceNormalizationCS.h"
 #include "NNXCore.h"
-#include "NNXTypes.h"
+#include "NNECoreTensor.h"
 
 namespace UE::NNEHlslShaders::Internal
 {
@@ -23,13 +23,13 @@ namespace UE::NNEHlslShaders::Internal
 		}
 	}
 
-	void TInstanceNormalizationCS::FillInParameters(float Epsilon, const NNX::FTensor& Input, TInstanceNormalizationCS::FParameters& Parameters)
+	void TInstanceNormalizationCS::FillInParameters(float Epsilon, const  UE::NNECore::Internal::FTensor& Input, TInstanceNormalizationCS::FParameters& Parameters)
 	{
 		check(Input.GetShape().Rank() >= 3);
 		
 		Parameters.Epsilon = Epsilon;
-		const int32 N = Input.GetShape().Data[0];
-		Parameters.C = Input.GetShape().Data[1];
+		const int32 N = Input.GetShape().GetData()[0];
+		Parameters.C = Input.GetShape().GetData()[1];
 		Parameters.NxC = N * Parameters.C;
 		Parameters.W = Input.GetShape().Volume() / (Parameters.NxC);
 	}
diff --git a/Engine/Plugins/Experimental/NNE/Source/NNXCore/Internal/NNECoreTensor.h b/Engine/Plugins/Experimental/NNE/Source/NNXCore/Internal/NNECoreTensor.h
new file mode 100644
--- /dev/null
+++ b/Engine/Plugins/Experimental/NNE/Source/NNXCore/Internal/NNECoreTensor.h
@@ -0,0 +1,107 @@
+// Copyright Epic Games, Inc. All Rights Reserved.
+
+#pragma once
+
+#include "CoreMinimal.h"
+#include "NNECoreTypes.h"
+
+namespace UE::NNECore::Internal
+{
+	/** Concrete tensor with data accessible by graph scheduling */
+	class FTensor
+	{
+	protected:
+		FString				Name;
+		ENNETensorDataType	DataType;
+		FTensorShape		Shape;
+		TArray<uint8>		PreparedData;
+		uint64				DataSize = 0;
+		uint32				Volume = 0;
+
+		FTensor() = default;
+		
+	public:
+		const FString& GetName() const
+		{
+			return Name;
+		}
+
+		ENNETensorDataType GetDataType() const
+		{
+			return DataType;
+		}
+
+		uint32 GetElemByteSize() const
+		{
+			return GetTensorDataTypeSizeInBytes(DataType);
+		}
+		
+		const FTensorShape& GetShape() const
+		{
+			return Shape;
+		}
+
+		template <typename T> TConstArrayView<T> GetPreparedData() const
+		{
+			const T* DataPtr = reinterpret_cast<const T*>(PreparedData.GetData());
+			const int32 ElemSize = sizeof(T);
+
+			check(PreparedData.Num() % ElemSize == 0);
+			return MakeArrayView(DataPtr, PreparedData.Num() / ElemSize);
+		}
+
+		void SetShape(const FTensorShape& InShape)
+		{
+			checkf(!HasPreparedData(), TEXT("Shape cannot be changed once data as been set."));
+			check(InShape.Volume() <= TNumericLimits<uint32>::Max());
+			Shape = InShape;
+			Volume = InShape.Volume();
+			DataSize = static_cast<uint64>(GetTensorDataTypeSizeInBytes(DataType) * Volume);
+		}
+
+		void SetPreparedData(TConstArrayView<uint8> Data)
+		{
+			checkf(Data.Num() == DataSize, TEXT("Incorrect data size, it should match tensor shape and data type."));
+			PreparedData.Empty();
+			PreparedData.Append(Data);
+		}
+
+		bool HasPreparedData() const
+		{
+			return !PreparedData.IsEmpty();
+		}
+
+		uint32 GetVolume() const
+		{
+			return Volume;
+		}
+
+		uint64 GetDataSize() const
+		{
+			return DataSize;
+		}
+
+		static FTensor Make(const FString& Name, const FTensorShape& Shape, ENNETensorDataType DataType)
+		{
+			FTensor Tensor;
+			Tensor.Name = Name;
+			Tensor.DataType = DataType;
+			Tensor.SetShape(Shape);
+			return Tensor;
+		}
+
+		static FTensor Make(const FTensorDesc& TensorDesc, const FTensorShape& Shape)
+		{
+			check(Shape.IsCompatibleWith(TensorDesc.GetShape()));
+			return Make(TensorDesc.GetName(), Shape, TensorDesc.GetDataType());
+		}
+
+		static FTensor MakeFromSymbolicDesc(const FTensorDesc& TensorDesc)
+		{
+			return Make(TensorDesc.GetName(), FTensorShape::MakeFromSymbolic(TensorDesc.GetShape()), TensorDesc.GetDataType());
+		}
+	};
+
+	using FTensorRef = FTensor*;
+
+} // namespace UE::NNECore::Internal
\ No newline at end of file
diff --git a/Engine/Plugins/Experimental/NNE/Source/NNXCore/Private/NNECoreModel.cpp b/Engine/Plugins/Experimental/NNE/Source/NNXCore/Private/NNECoreModel.cpp
--- a/Engine/Plugins/Experimental/NNE/Source/NNXCore/Private/NNECoreModel.cpp
+++ b/Engine/Plugins/Experimental/NNE/Source/NNXCore/Private/NNECoreModel.cpp
@@ -113,7 +113,7 @@ TArray<int32> UNNEModel::GetInputShapes(int32 Index)
 		return TArray<int32>();
 	}
 
-	return TArray<int32>(Desc[Index].GetShape().Data);
+	return TArray<int32>(Desc[Index].GetShape().GetData());
 }
 
 TArray<int32> UNNEModel::GetOutputShapes(int32 Index)
@@ -131,7 +131,7 @@ TArray<int32> UNNEModel::GetOutputShapes(int32 Index)
 		return TArray<int32>();
 	}
 
-	return TArray<int32>(Desc[Index].GetShape().Data);
+	return TArray<int32>(Desc[Index].GetShape().GetData());
 }
 
 bool UNNEModel::SetInput(const TArray<FNNETensor>& Input)
@@ -156,7 +156,7 @@ bool UNNEModel::SetInput(const TArray<FNNETensor>& Input)
 	for (int32 i = 0; i < Input.Num(); i++)
 	{
 		InputBindings.Add(NNX::FMLTensorBinding::FromCPU((void*)(Input[i].Data.GetData()), sizeof(float) * Input[i].Data.Num()));
-		InputShapes.Add(NNX::FTensorShape::MakeFromSymbolic(NNX::FSymbolicTensorShape::Make(Input[i].Shape)));
+		InputShapes.Add(NNX::FTensorShape::MakeFromSymbolic(UE::NNECore::FSymbolicTensorShape::Make(Input[i].Shape)));
 	}
 
 	if (Model->SetInputTensorShapes(InputShapes) != 0)
@@ -238,7 +238,7 @@ bool UNNEModel::RunAsync(FNNETaskPriority TaskPriority, FNNEModelOnAsyncResult O
 		for (int32 i = 0; i < OutputDescs.Num(); i++)
 		{
 			OutputTensors.Add(FNNETensor());
-			OutputTensors[i].Shape = OutputDescs[i].GetShape().Data;
+			OutputTensors[i].Shape = OutputDescs[i].GetShape().GetData();
 			int32 Size = OutputTensors[i].Shape.Num() > 0 ? 1 : 0;
 			for (int32 j = 0; j < OutputTensors[i].Shape.Num(); j++)
 			{
diff --git a/Engine/Plugins/Experimental/NNE/Source/NNXCore/Private/NNECoreTypes.cpp b/Engine/Plugins/Experimental/NNE/Source/NNXCore/Private/NNECoreTypes.cpp
new file mode 100644
--- /dev/null
+++ b/Engine/Plugins/Experimental/NNE/Source/NNXCore/Private/NNECoreTypes.cpp
@@ -0,0 +1,178 @@
+// Copyright Epic Games, Inc. All Rights Reserved.
+
+#pragma once
+
+#include "NNECoreTypes.h"
+#include "NNXCore.h"
+
+namespace UE::NNECore::TypesUtils
+{
+	template <class T> bool AreShapeEqual(const T& lhs, const T& rhs)
+	{
+		check(lhs.Rank() <= lhs.MaxRank);
+
+		if (lhs.Rank() != rhs.Rank())
+			return false;
+
+		for (int32 i = 0; i < lhs.Rank(); ++i)
+		{
+			if (lhs.GetData()[i] != rhs.GetData()[i])
+				return false;
+		}
+			
+		return true;
+	}
+}
+
+namespace UE::NNECore
+{
+
+bool FSymbolicTensorShape::IsConcrete() const
+{
+	for (int32 i = 0; i < Rank(); ++i)
+	{
+		if (Data[i] < 0)
+		{
+			return false;
+		}
+	}
+	return true;
+}
+
+bool FSymbolicTensorShape::operator==(const FSymbolicTensorShape& OtherShape) const
+{
+	return TypesUtils::AreShapeEqual(*this, OtherShape);
+}
+
+bool FSymbolicTensorShape::operator!=(const FSymbolicTensorShape& OtherShape) const { return !(*this == OtherShape); }
+
+void FSymbolicTensorShape::operator=(const FSymbolicTensorShape& OtherShape)
+{
+	check(OtherShape.Rank() <= MaxRank);
+	Data = OtherShape.Data;
+}
+
+FSymbolicTensorShape FSymbolicTensorShape::Make(TConstArrayView<int32> Data)
+{
+	if (Data.Num() > MaxRank)
+	{
+		UE_LOG(LogNNX, Warning, TEXT("Cannot create symbolic tensor shape, input is rank %d while max rank is %d"), Data.Num(), MaxRank);
+		return {};
+	}
+
+	FSymbolicTensorShape Shape;
+	Shape.Data.Append(Data);
+	return Shape;
+}
+
+uint64 FTensorShape::Volume() const
+{
+	uint64 Result = 1;
+
+	for (int32 Idx = 0; Idx < Data.Num(); ++Idx)
+	{
+		Result *= Data[Idx];
+	}
+
+	return Result;
+};
+
+bool FTensorShape::IsCompatibleWith(const FSymbolicTensorShape& SymbolicShape) const
+{
+	if (Rank() != SymbolicShape.Rank())
+	{
+		return false;
+	}
+	for (int32 i = 0; i < Rank(); ++i)
+	{
+		if (SymbolicShape.GetData()[i] >= 0 && SymbolicShape.GetData()[i] != Data[i])
+		{
+			return false;
+		}
+	}
+	return true;
+}
+
+bool FTensorShape::operator==(const FTensorShape& OtherShape) const
+{
+	return TypesUtils::AreShapeEqual(*this, OtherShape);
+}
+
+bool FTensorShape::operator!=(const FTensorShape& OtherShape) const { return !(*this == OtherShape); }
+
+void FTensorShape::operator=(const FTensorShape& OtherShape)
+{
+	check(OtherShape.Rank() <= MaxRank);
+	Data = OtherShape.Data;
+}
+
+FTensorShape FTensorShape::Make(TConstArrayView<uint32> Data)
+{
+	if (Data.Num() > MaxRank)
+	{
+		UE_LOG(LogNNX, Warning, TEXT("Cannot create tensor shape, input is rank %d while max rank is %d"), Data.Num(), MaxRank);
+		return {};
+	}
+
+	FTensorShape Shape;
+	Shape.Data.Append(Data);
+	return Shape;
+}
+	
+FTensorShape FTensorShape::MakeFromSymbolic(const FSymbolicTensorShape& SymbolicShape)
+{
+	FTensorShape ConcreteShape;
+	for (int32 i = 0; i < SymbolicShape.GetData().Num(); ++i)
+	{
+		int32 Dim = SymbolicShape.GetData()[i];
+		ConcreteShape.Data.Add(Dim < 0 ? 1 : Dim);
+	}
+	return ConcreteShape;
+}
+
+size_t GetTensorDataTypeSizeInBytes(ENNETensorDataType InType)
+{
+	switch (InType)
+	{
+	case ENNETensorDataType::Complex128:
+		return 16;
+
+	case ENNETensorDataType::Complex64:
+		return 8;
+
+	case ENNETensorDataType::Double:
+	case ENNETensorDataType::Int64:
+	case ENNETensorDataType::UInt64:
+		return 8;
+
+	case ENNETensorDataType::Float:
+	case ENNETensorDataType::Int32:
+	case ENNETensorDataType::UInt32:
+		return 4;
+
+	case ENNETensorDataType::Half:
+	case ENNETensorDataType::BFloat16:
+	case ENNETensorDataType::Int16:
+	case ENNETensorDataType::UInt16:
+		return 2;
+
+	case ENNETensorDataType::Int8:
+	case ENNETensorDataType::UInt8:
+	case ENNETensorDataType::Char:
+	case ENNETensorDataType::Boolean:
+		return 1;
+	}
+
+	return 0;
+}
+
+FTensorDesc FTensorDesc::Make(const FString& Name, const FSymbolicTensorShape& Shape, ENNETensorDataType DataType)
+{
+	FTensorDesc Desc;
+	Desc.Name = Name;
+	Desc.DataType = DataType;
+	Desc.Shape = Shape;
+	return Desc;
+}
+
+} // namespace UE::NNECore
diff --git a/Engine/Plugins/Experimental/NNE/Source/NNXCore/Public/NNECoreTypes.h b/Engine/Plugins/Experimental/NNE/Source/NNXCore/Public/NNECoreTypes.h
new file mode 100644
--- /dev/null
+++ b/Engine/Plugins/Experimental/NNE/Source/NNXCore/Public/NNECoreTypes.h
@@ -0,0 +1,94 @@
+// Copyright Epic Games, Inc. All Rights Reserved.
+
+#pragma once
+
+#include "CoreMinimal.h"
+#include "UObject/Class.h"
+
+UENUM()
+enum class ENNETensorDataType : uint8
+{
+	None,
+	Char,								//!< Character type
+	Boolean,							//!< Boolean type
+	Half,								//!< 16-bit floating number
+	Float,								//!< 32-bit floating number
+	Double,								//!< 64-bit floating number
+	Int8,								//!< 8-bit signed integer
+	Int16,								//!< 16-bit signed integer
+	Int32,								//!< 32-bit signed integer
+	Int64,								//!< 64-bit signed integer
+	UInt8,								//!< 8-bit unsigned integer
+	UInt16,								//!< 16-bit unsigned integer
+	UInt32,								//!< 32-bit unsigned integer
+	UInt64,								//!< 64-bit unsigned integer
+	Complex64,							//!< 64-bit Complex Number
+	Complex128,							//!< 128-bit Complex Number
+	BFloat16							//!< 16-bit floating number
+};
+
+namespace UE::NNECore
+{
+	class NNXCORE_API FSymbolicTensorShape
+	{
+	public:
+		constexpr static int32	MaxRank = 8;
+
+	private:
+		TArray<int32, TInlineAllocator<MaxRank>> Data;
+
+	public:
+		static FSymbolicTensorShape Make(TConstArrayView<int32> Data);
+
+		inline TConstArrayView<int32> GetData() const { return Data; };
+		inline int32 Rank() const { return Data.Num(); }
+
+		bool IsConcrete() const;
+		bool operator==(const FSymbolicTensorShape& OtherShape) const;
+		bool operator!=(const FSymbolicTensorShape& OtherShape) const;
+		void operator=(const FSymbolicTensorShape& OtherShape);
+	};
+
+	class NNXCORE_API FTensorShape
+	{
+	public:
+		constexpr static int32	MaxRank = FSymbolicTensorShape::MaxRank;
+
+	private:
+		TArray<uint32, TInlineAllocator<MaxRank>> Data;
+
+	public:
+		static FTensorShape Make(TConstArrayView<uint32> Data);
+		static FTensorShape MakeFromSymbolic(const FSymbolicTensorShape& SymbolicShape);
+
+		inline TConstArrayView<uint32> GetData() const { return Data; };
+		inline int32 Rank() const { return Data.Num(); }
+		
+		uint64 Volume() const;
+		bool IsCompatibleWith(const FSymbolicTensorShape& SymbolicShape) const;
+		bool operator==(const FTensorShape& OtherShape) const;
+		bool operator!=(const FTensorShape& OtherShape) const;
+		void operator=(const FTensorShape& OtherShape);
+	};
+
+	/** Return data size in bytes for tensor data type */
+	size_t NNXCORE_API GetTensorDataTypeSizeInBytes(ENNETensorDataType InType);
+
+	class NNXCORE_API FTensorDesc
+	{
+		FString					Name;
+		ENNETensorDataType		DataType;
+		FSymbolicTensorShape	Shape;
+
+		FTensorDesc() = default;
+
+	public:
+		static FTensorDesc Make(const FString& Name, const FSymbolicTensorShape& Shape, ENNETensorDataType DataType);
+		
+		inline const FString& GetName() const { return Name; }
+		inline ENNETensorDataType GetDataType() const { return DataType; }
+		inline uint32 GetElemByteSize() const { return GetTensorDataTypeSizeInBytes(DataType); }
+		inline const FSymbolicTensorShape& GetShape() const { return Shape; }
+	};
+
+} // namespace NNE
diff --git a/Engine/Plugins/Experimental/NNE/Source/NNXCore/Public/NNXRuntime.h b/Engine/Plugins/Experimental/NNE/Source/NNXCore/Public/NNXRuntime.h
--- a/Engine/Plugins/Experimental/NNE/Source/NNXCore/Public/NNXRuntime.h
+++ b/Engine/Plugins/Experimental/NNE/Source/NNXCore/Public/NNXRuntime.h
@@ -2,28 +2,19 @@
 
 #pragma once
 
-#include "NNXTypes.h"
+#include "NNECoreTypes.h"
 #include "CoreMinimal.h"
 
 class FRDGBuffer;
-
-// TODO: 
-/*
-- OperatorRegistry
-	- Q: Should registry be an internal part of runtime ModelOptimizer, i.e. never exposed to the NNX utils?
-
-- ModelOptimizer interface for runtime dependent model optimizations
-*/
-
-class UMLInferenceModel;
 class FRDGBuilder;
+class UMLInferenceModel;
 
 namespace NNX
 {
-class FMLInferenceModel;
-//class IModelOptimizer;
+using FTensorDesc = UE::NNECore::FTensorDesc;
+using FTensorShape = UE::NNECore::FTensorShape;
 
-//class IOperatorRegistry;
+class FMLInferenceModel;
 
 /** Runtime support flags */
 enum class EMLRuntimeSupportFlags : uint32
diff --git a/Engine/Plugins/Experimental/NNE/Source/NNXCore/Public/NNXRuntimeFormat.h b/Engine/Plugins/Experimental/NNE/Source/NNXCore/Public/NNXRuntimeFormat.h
--- a/Engine/Plugins/Experimental/NNE/Source/NNXCore/Public/NNXRuntimeFormat.h
+++ b/Engine/Plugins/Experimental/NNE/Source/NNXCore/Public/NNXRuntimeFormat.h
@@ -4,7 +4,7 @@
 
 #include "UObject/Object.h"
 #include "NNECoreAttributeValue.h"
-#include "NNXTypes.h"
+#include "NNECoreTypes.h"
 #include "NNXRuntime.h"
 #include "NNXRuntimeFormat.generated.h"
 
@@ -88,7 +88,7 @@ struct FMLFormatTensorDesc
 	EMLFormatTensorType	Type = EMLFormatTensorType::None;
 
 	UPROPERTY(VisibleAnywhere, Category = "Neural Network Inference")
-	EMLTensorDataType	DataType = EMLTensorDataType::None;
+	ENNETensorDataType	DataType = ENNETensorDataType::None;
 
 	UPROPERTY(VisibleAnywhere, Category = "Neural Network Inference")
 	uint64	DataSize = 0;
diff --git a/Engine/Plugins/Experimental/NNE/Source/NNXCore/Public/NNXTypes.h b/Engine/Plugins/Experimental/NNE/Source/NNXCore/Public/NNXTypes.h
deleted file mode 100644
diff --git a/Engine/Plugins/Experimental/NNE/Source/NNXQA/Private/NNXQAJsonUtils.h b/Engine/Plugins/Experimental/NNE/Source/NNXQA/Private/NNXQAJsonUtils.h
--- a/Engine/Plugins/Experimental/NNE/Source/NNXQA/Private/NNXQAJsonUtils.h
+++ b/Engine/Plugins/Experimental/NNE/Source/NNXQA/Private/NNXQAJsonUtils.h
@@ -7,7 +7,6 @@
 #include "Serialization/JsonSerializerMacros.h"
 #include "Dom/JsonObject.h"
 #include "Misc/FileHelper.h"
-#include "NNXTypes.h"
 
 namespace NNX 
 {
diff --git a/Engine/Plugins/Experimental/NNE/Source/NNXQA/Private/NNXQAParametricTest.cpp b/Engine/Plugins/Experimental/NNE/Source/NNXQA/Private/NNXQAParametricTest.cpp
--- a/Engine/Plugins/Experimental/NNE/Source/NNXQA/Private/NNXQAParametricTest.cpp
+++ b/Engine/Plugins/Experimental/NNE/Source/NNXQA/Private/NNXQAParametricTest.cpp
@@ -3,7 +3,7 @@
 #include "NNXQAParametricTest.h"
 
 #include "NNXCore.h"
-#include "NNXTypes.h"
+#include "NNECoreTypes.h"
 #include "NNXRuntimeFormat.h"
 #include "NNECoreAttributeMap.h"
 #include "NNXQAUtils.h"
@@ -21,6 +21,8 @@
 
 namespace NNX 
 {
+using FSymbolicTensorShape = UE::NNECore::FSymbolicTensorShape;
+	
 namespace Test
 {
 	class FParametricTests : public FTests
@@ -127,11 +129,11 @@ namespace Test
 			return ConcreteShape;
 		}
 
-		static EMLTensorDataType GetTensorTypeFromJson(const FString& TypeName, EMLTensorDataType DefaultValue)
+		static ENNETensorDataType GetTensorTypeFromJson(const FString& TypeName, ENNETensorDataType DefaultValue)
 		{
-			int64 Value = StaticEnum<EMLTensorDataType>()->GetValueByNameString(TypeName);
+			int64 Value = StaticEnum<ENNETensorDataType>()->GetValueByNameString(TypeName);
 			
-			return (Value == INDEX_NONE) ? DefaultValue : (EMLTensorDataType)Value;
+			return (Value == INDEX_NONE) ? DefaultValue : (ENNETensorDataType)Value;
 		}
 
 		class ElementWiseFromJsonStringTensorInitializer
@@ -185,7 +187,7 @@ namespace Test
 			return TArray<char>();
 		}
 
-		static void ApplyDatasetConfig(FTests::FTestSetup& TestSetup, const Json::FTestConfigDataset& TestDataset, EMLTensorDataType DefaultInputType, EMLTensorDataType DefaultOutputType)
+		static void ApplyDatasetConfig(FTests::FTestSetup& TestSetup, const Json::FTestConfigDataset& TestDataset, ENNETensorDataType DefaultInputType, ENNETensorDataType DefaultOutputType)
 		{
 			ApplyRuntimesConfig(TestSetup, TestDataset.Runtimes);
 
@@ -198,7 +200,7 @@ namespace Test
 			for (auto&& JsonTensor : TestDataset.Inputs)
 			{
 				FTensorShape Shape = GetConcreteShapeFromJsonArray(JsonTensor.Shape);
-				EMLTensorDataType TensorType = GetTensorTypeFromJson(JsonTensor.Type, DefaultInputType);
+				ENNETensorDataType TensorType = GetTensorTypeFromJson(JsonTensor.Type, DefaultInputType);
 				FTensor Tensor = FTensor::Make(FString::Printf(TEXT("input%d"), i), Shape, TensorType);
 				TArray<char> TensorData = GetInputTensorDataFromJson(Tensor, i, JsonTensor.Source);
 
@@ -210,7 +212,7 @@ namespace Test
 			for (auto&& JsonTensor : TestDataset.Weights)
 			{
 				FTensorShape Shape = GetConcreteShapeFromJsonArray(JsonTensor.Shape);
-				EMLTensorDataType TensorType = GetTensorTypeFromJson(JsonTensor.Type, EMLTensorDataType::Float);
+				ENNETensorDataType TensorType = GetTensorTypeFromJson(JsonTensor.Type, ENNETensorDataType::Float);
 				FTensor Tensor = FTensor::Make(FString::Printf(TEXT("weights%d"), i), Shape, TensorType);
 				TArray<char> TensorData = GetInputTensorDataFromJson(Tensor, i, JsonTensor.Source);
 
@@ -222,7 +224,7 @@ namespace Test
 			for (auto&& JsonTensor : TestDataset.Outputs)
 			{
 				FTensorShape Shape = GetConcreteShapeFromJsonArray(JsonTensor.Shape);
-				EMLTensorDataType TensorType = GetTensorTypeFromJson(JsonTensor.Type, DefaultOutputType);
+				ENNETensorDataType TensorType = GetTensorTypeFromJson(JsonTensor.Type, DefaultOutputType);
 				FTensor Tensor = FTensor::Make(FString::Printf(TEXT("output%d"), i), Shape, TensorType);
 				TArray<char> TensorData = GetOutputTensorDataFromJson(Tensor, i, JsonTensor.Source);
 
@@ -234,7 +236,7 @@ namespace Test
 			if (TestDataset.Outputs.Num() == 0 && TestDataset.Inputs.Num() > 0)
 			{
 				FTensorShape Shape = GetConcreteShapeFromJsonArray(TestDataset.Inputs[0].Shape);
-				EMLTensorDataType TensorType = GetTensorTypeFromJson(TestDataset.Inputs[0].Type, DefaultOutputType);
+				ENNETensorDataType TensorType = GetTensorTypeFromJson(TestDataset.Inputs[0].Type, DefaultOutputType);
 				FTensor Tensor = FTensor::Make(FString::Printf(TEXT("output"), i), Shape, TensorType);
 				TArray<char> TensorData = GetOutputTensorDataFromJson(Tensor, i, TestDataset.Inputs[0].Source);
 				
@@ -359,8 +361,8 @@ namespace Test
 					}
 
 					const FString& TestBaseName = TestTarget.Target;
-					EMLTensorDataType InputTypeFromTarget = GetTensorTypeFromJson(TestTarget.InputType, EMLTensorDataType::Float);
-					EMLTensorDataType OutputTypeFromTarget = GetTensorTypeFromJson(TestTarget.OutputType, EMLTensorDataType::Float);
+					ENNETensorDataType InputTypeFromTarget = GetTensorTypeFromJson(TestTarget.InputType, ENNETensorDataType::Float);
+					ENNETensorDataType OutputTypeFromTarget = GetTensorTypeFromJson(TestTarget.OutputType, ENNETensorDataType::Float);
 					bool bAtLeastATestWasAdded = false;
 					
 					for (const Json::FTestConfigInputOutputSet& InputOutputSet : InputOutputSets)
diff --git a/Engine/Plugins/Experimental/NNE/Source/NNXQA/Private/NNXQAUtils.cpp b/Engine/Plugins/Experimental/NNE/Source/NNXQA/Private/NNXQAUtils.cpp
--- a/Engine/Plugins/Experimental/NNE/Source/NNXQA/Private/NNXQAUtils.cpp
+++ b/Engine/Plugins/Experimental/NNE/Source/NNXQA/Private/NNXQAUtils.cpp
@@ -13,22 +13,22 @@ namespace NNX
 namespace Test 
 {
 	static void FillTensors(
-		TConstArrayView<FTensor> TensorsFromTestSetup,
-		TConstArrayView<FTensorDesc> TensorDescsFromModel,
-		TArray<FTensor>& OutTensors)
+		TConstArrayView<UE::NNECore::Internal::FTensor> TensorsFromTestSetup,
+		TConstArrayView<UE::NNECore::FTensorDesc> TensorDescsFromModel,
+		TArray<UE::NNECore::Internal::FTensor>& OutTensors)
 	{
 		OutTensors = TensorsFromTestSetup;
 		if (OutTensors.IsEmpty())
 		{
 			for (const FTensorDesc& SymbolicTensorDesc : TensorDescsFromModel)
 			{
-				FTensor TensorDesc = FTensor::MakeFromSymbolicDesc(SymbolicTensorDesc);
+				UE::NNECore::Internal::FTensor TensorDesc = UE::NNECore::Internal::FTensor::MakeFromSymbolicDesc(SymbolicTensorDesc);
 				OutTensors.Emplace(TensorDesc);
 			}
 		}
 	}
 	
-	ElementWiseCosTensorInitializer::ElementWiseCosTensorInitializer(EMLTensorDataType InDataType, uint32 InTensorIndex)
+	ElementWiseCosTensorInitializer::ElementWiseCosTensorInitializer(ENNETensorDataType InDataType, uint32 InTensorIndex)
 	: DataType(InDataType), TensorIndex(InTensorIndex)
 	{
 	}
@@ -39,37 +39,37 @@ namespace Test
 		switch (DataType)
 		{
 
-		case EMLTensorDataType::Boolean:
+		case ENNETensorDataType::Boolean:
 			return (ElementIndex + IndexOffsetBetweenTensor * TensorIndex) % 2;
-		case EMLTensorDataType::Char:
-		case EMLTensorDataType::Int8:
-		case EMLTensorDataType::Int16:
-		case EMLTensorDataType::Int32:
-		case EMLTensorDataType::Int64:
+		case ENNETensorDataType::Char:
+		case ENNETensorDataType::Int8:
+		case ENNETensorDataType::Int16:
+		case ENNETensorDataType::Int32:
+		case ENNETensorDataType::Int64:
 			return 10.0f * FMath::Cos((float)(ElementIndex + IndexOffsetBetweenTensor * TensorIndex));
-		case EMLTensorDataType::UInt8:
-		case EMLTensorDataType::UInt16:
-		case EMLTensorDataType::UInt32:
-		case EMLTensorDataType::UInt64:
+		case ENNETensorDataType::UInt8:
+		case ENNETensorDataType::UInt16:
+		case ENNETensorDataType::UInt32:
+		case ENNETensorDataType::UInt64:
 			return 10.0f * FMath::Abs(FMath::Cos((float)(ElementIndex + IndexOffsetBetweenTensor * TensorIndex)));
 		default:
-			//case EMLTensorDataType::None:
-			//case EMLTensorDataType::Half:
-			//case EMLTensorDataType::Double:
-			//case EMLTensorDataType::Float:
-			//case EMLTensorDataType::Complex64:
-			//case EMLTensorDataType::Complex128:
-			//case EMLTensorDataType::BFloat16:
+			//case ENNETensorDataType::None:
+			//case ENNETensorDataType::Half:
+			//case ENNETensorDataType::Double:
+			//case ENNETensorDataType::Float:
+			//case ENNETensorDataType::Complex64:
+			//case ENNETensorDataType::Complex128:
+			//case ENNETensorDataType::BFloat16:
 			return FMath::Cos((float)(ElementIndex + IndexOffsetBetweenTensor * TensorIndex));
 		}
 	};
 
-	TArray<char> GenerateTensorDataForTest(const FTensor& Tensor, std::function<float(uint32)> ElementInitializer)
+	TArray<char> GenerateTensorDataForTest(const UE::NNECore::Internal::FTensor& Tensor, std::function<float(uint32)> ElementInitializer)
 	{
 		const uint32 NumberOfElements = Tensor.GetVolume();
 		const int32 ElementByteSize = Tensor.GetElemByteSize();
 		const uint64 BufferSize = Tensor.GetDataSize();
-		const EMLTensorDataType DataType = Tensor.GetDataType();
+		const ENNETensorDataType DataType = Tensor.GetDataType();
 		TArray<char> Buffer;
 			
 		Buffer.SetNum(BufferSize);
@@ -84,19 +84,19 @@ namespace Test
 			const char BoolData = FloatData != 0.0f ? 1 : 0;
 			switch (DataType)
 			{
-			case EMLTensorDataType::Float:
+			case ENNETensorDataType::Float:
 				check(sizeof(FloatData) == ElementByteSize);
 				FMemory::Memcpy(DestinationPtr, &FloatData, ElementByteSize); break;
-			case EMLTensorDataType::Int32:
+			case ENNETensorDataType::Int32:
 				check(sizeof(IntData) == ElementByteSize);
 				FMemory::Memcpy(DestinationPtr, &IntData, ElementByteSize); break;
-			case EMLTensorDataType::Int64:
+			case ENNETensorDataType::Int64:
 				check(sizeof(Int64Data) == ElementByteSize);
 				FMemory::Memcpy(DestinationPtr, &Int64Data, ElementByteSize); break;
-			case EMLTensorDataType::UInt32:
+			case ENNETensorDataType::UInt32:
 				check(sizeof(UIntData) == ElementByteSize);
 				FMemory::Memcpy(DestinationPtr, &UIntData, ElementByteSize); break;
-			case EMLTensorDataType::Boolean:
+			case ENNETensorDataType::Boolean:
 				check(sizeof(BoolData) == ElementByteSize);
 				FMemory::Memcpy(DestinationPtr, &BoolData, ElementByteSize); break;
 			default:
@@ -108,7 +108,7 @@ namespace Test
 		return Buffer;
 	}
 
-	static void FillInputTensorBindings(TConstArrayView<FTensor> Tensors, TConstArrayView<FTests::FTensorData> TensorsData, 
+	static void FillInputTensorBindings(TConstArrayView<UE::NNECore::Internal::FTensor> Tensors, TConstArrayView<FTests::FTensorData> TensorsData,
 		TArray<TArray<char>>& OutMemBuffers, TArray<FMLTensorBinding>& OutBindings)
 	{
 		OutBindings.Empty();
@@ -117,7 +117,7 @@ namespace Test
 
 		for (int Index = 0; Index < Tensors.Num(); ++Index)
 		{
-			const FTensor& Tensor = Tensors[Index];
+			const UE::NNECore::Internal::FTensor& Tensor = Tensors[Index];
 			TArray<char>& MemBuffer = OutMemBuffers.Emplace_GetRef();
 			if (TensorsData.Num() == 0 || TensorsData[Index].Num() == 0)
 			{
@@ -133,7 +133,7 @@ namespace Test
 		}
 	}
 
-	static void FillOutputTensorBindings(TConstArrayView<FTensor> Tensors, TArray<TArray<char>>& OutMemBuffers, TArray<FMLTensorBinding>& OutBindings)
+	static void FillOutputTensorBindings(TConstArrayView<UE::NNECore::Internal::FTensor> Tensors, TArray<TArray<char>>& OutMemBuffers, TArray<FMLTensorBinding>& OutBindings)
 	{
 		OutBindings.Empty();
 		OutMemBuffers.Empty();
@@ -163,21 +163,21 @@ namespace Test
 	template FString ShapeToString<int32>(TConstArrayView<int32> Shape);
 	template FString ShapeToString<uint32>(TConstArrayView<uint32> Shape);
 
-	FString TensorToString(const FTensor& desc)
+	FString TensorToString(const UE::NNECore::Internal::FTensor& desc)
 	{
-		TArrayView<const uint32> Shape(desc.GetShape().Data);
+		TArrayView<const uint32> Shape(desc.GetShape().GetData());
 
 		FString TensorDesc;
 		TensorDesc.Reserve(50);
 		TensorDesc += FString::Printf(TEXT("Name: %s, Shape: "), *desc.GetName());
 		TensorDesc += ShapeToString(Shape);
-		const FString& DataTypeName = StaticEnum<EMLTensorDataType>()->GetNameStringByValue(static_cast<int64>(desc.GetDataType()));
+		const FString& DataTypeName = StaticEnum<ENNETensorDataType>()->GetNameStringByValue(static_cast<int64>(desc.GetDataType()));
 		TensorDesc += FString::Printf(TEXT(" DataType: %s"), *DataTypeName);
 
 		return TensorDesc;
 	}
 
-	FString TensorToString(const FTensor& TensorDesc, TConstArrayView<char> TensorData)
+	FString TensorToString(const UE::NNECore::Internal::FTensor& TensorDesc, TConstArrayView<char> TensorData)
 	{
 		FString TensorLog;
 		TensorLog.Reserve(50);
@@ -196,13 +196,13 @@ namespace Test
 
 			switch (TensorDesc.GetDataType())
 			{
-				case EMLTensorDataType::Float:
+				case ENNETensorDataType::Float:
 					TensorLog += FString::Printf(TEXT("%0.2f"), *(float*)Data); break;
-				case EMLTensorDataType::Int32:
+				case ENNETensorDataType::Int32:
 					TensorLog += FString::Printf(TEXT("%d"), *(int32*)Data); break;
-				case EMLTensorDataType::UInt32:
+				case ENNETensorDataType::UInt32:
 					TensorLog += FString::Printf(TEXT("%u"), *(uint32*)Data); break;
-				case EMLTensorDataType::Boolean:
+				case ENNETensorDataType::Boolean:
 					check(ElementByteSize == 1);
 					TensorLog += FString::Printf(TEXT("%s"), *(bool*)Data ? "true" : "false"); break;
 				default:
@@ -218,8 +218,8 @@ namespace Test
 	}
 
 	template<typename T> bool CompareTensorData(
-		const FTensor& RefTensorDesc,   const TArray<char>& RefRawBuffer,
-		const FTensor& OtherTensorDesc, const TArray<char>& OtherRawBuffer,
+		const UE::NNECore::Internal::FTensor& RefTensorDesc,   const TArray<char>& RefRawBuffer,
+		const UE::NNECore::Internal::FTensor& OtherTensorDesc, const TArray<char>& OtherRawBuffer,
 		float AbsoluteErrorEpsilon, float RelativeErrorPercent)
 	{
 		const T* RefBuffer = (T*)RefRawBuffer.GetData();
@@ -320,32 +320,32 @@ namespace Test
 	}
 
 	bool VerifyTensorResult(
-		const FTensor& RefTensor, const TArray<char>& RefRawBuffer,
-		const FTensor& OtherTensor, const TArray<char>& OtherRawBuffer,
+		const UE::NNECore::Internal::FTensor& RefTensor, const TArray<char>& RefRawBuffer,
+		const UE::NNECore::Internal::FTensor& OtherTensor, const TArray<char>& OtherRawBuffer,
 		float AbsoluteErrorEpsilon, float RelativeErrorPercent)
 	{
 		if (RefTensor.GetShape() != OtherTensor.GetShape())
 		{
-			TArrayView<const uint32> RefShape(RefTensor.GetShape().Data);
-			TArrayView<const uint32> OtherShape(OtherTensor.GetShape().Data);
+			TArrayView<const uint32> RefShape(RefTensor.GetShape().GetData());
+			TArrayView<const uint32> OtherShape(OtherTensor.GetShape().GetData());
 			UE_LOG(LogNNX, Error, TEXT("Tensor shape do not match.\nExpected: %s\nGot:      %s"), *ShapeToString(RefShape), *ShapeToString(OtherShape));
 			return false;
 		}
 
-		if (RefTensor.GetDataType() == EMLTensorDataType::Float)
+		if (RefTensor.GetDataType() == ENNETensorDataType::Float)
 		{
 			return CompareTensorData<float>(RefTensor, RefRawBuffer, OtherTensor, OtherRawBuffer, AbsoluteErrorEpsilon, RelativeErrorPercent);
 		}
-		else if (RefTensor.GetDataType() == EMLTensorDataType::Boolean)
+		else if (RefTensor.GetDataType() == ENNETensorDataType::Boolean)
 		{
 			check(RefTensor.GetElemByteSize() == 1);
 			return CompareTensorData<bool>(RefTensor, RefRawBuffer, OtherTensor, OtherRawBuffer, AbsoluteErrorEpsilon, RelativeErrorPercent);
 		}
-		else if (RefTensor.GetDataType() == EMLTensorDataType::Int32)
+		else if (RefTensor.GetDataType() == ENNETensorDataType::Int32)
 		{
 			return CompareTensorData<int32>(RefTensor, RefRawBuffer, OtherTensor, OtherRawBuffer, AbsoluteErrorEpsilon, RelativeErrorPercent);
 		}
-		else if (RefTensor.GetDataType() == EMLTensorDataType::UInt32)
+		else if (RefTensor.GetDataType() == ENNETensorDataType::UInt32)
 		{
 			return CompareTensorData<uint32>(RefTensor, RefRawBuffer, OtherTensor, OtherRawBuffer, AbsoluteErrorEpsilon, RelativeErrorPercent);
 		}
@@ -359,7 +359,7 @@ namespace Test
 	}
 
 	static bool RunTestInference(const FNNIModelRaw& ONNXModelData, const FTests::FTestSetup& TestSetup, 
-		IRuntime* Runtime, TArray<FTensor>& OutOutputTensors, TArray<TArray<char>>& OutOutputMemBuffers)
+		IRuntime* Runtime, TArray<UE::NNECore::Internal::FTensor>& OutOutputTensors, TArray<TArray<char>>& OutOutputMemBuffers)
 	{
 		OutOutputMemBuffers.Empty();
 
@@ -372,7 +372,7 @@ namespace Test
 		}
 
 		// If test does not ask for specific input/output, fill with model default converting variable dim to 1 if any.
-		TArray<FTensor> InputTensors;
+		TArray<UE::NNECore::Internal::FTensor> InputTensors;
 		
 		FillTensors(TestSetup.Inputs, InferenceModel->GetInputTensorDescs(), InputTensors);
 		FillTensors(TestSetup.Outputs, InferenceModel->GetOutputTensorDescs(), OutOutputTensors);
@@ -402,8 +402,8 @@ namespace Test
 		}
 		#endif //UE_BUILD_DEBUG
 
-		TArray<FTensorShape> InputShapes;
-		for (const FTensor& Tensor : InputTensors)
+		TArray<UE::NNECore::FTensorShape> InputShapes;
+		for (const UE::NNECore::Internal::FTensor& Tensor : InputTensors)
 		{
 			InputShapes.Emplace(Tensor.GetShape());
 		}
@@ -423,7 +423,7 @@ namespace Test
 		}
 
 		// Verify output shapes are as expected
-		TConstArrayView<FTensorShape> OutputShapes = InferenceModel->GetOutputTensorShapes();
+		TConstArrayView<UE::NNECore::FTensorShape> OutputShapes = InferenceModel->GetOutputTensorShapes();
 		if (OutputShapes.Num() != OutOutputTensors.Num())
 		{
 			UE_LOG(LogNNX, Error, TEXT("Expected %d output tensors, got %d."), OutOutputTensors.Num(), OutputShapes.Num());
@@ -431,12 +431,12 @@ namespace Test
 		}
 		for (int i = 0; i < OutOutputTensors.Num(); ++i)
 		{
-			const FTensorShape& RefTensorShape = OutOutputTensors[i].GetShape();
-			const FTensorShape& OtherTensorShape = OutputShapes[i];
+			const UE::NNECore::FTensorShape& RefTensorShape = OutOutputTensors[i].GetShape();
+			const UE::NNECore::FTensorShape& OtherTensorShape = OutputShapes[i];
 			if (RefTensorShape != OtherTensorShape)
 			{
-				TArrayView<const uint32> RefShape(RefTensorShape.Data);
-				TArrayView<const uint32> OtherShape(OtherTensorShape.Data);
+				TArrayView<const uint32> RefShape(RefTensorShape.GetData());
+				TArrayView<const uint32> OtherShape(OtherTensorShape.GetData());
 				UE_LOG(LogNNX, Error, TEXT("Output shape do not match at index %d.\nExpected: %s\nGot:      %s"), i, *ShapeToString(RefShape), *ShapeToString(OtherShape));
 				return false;
 			}
@@ -446,10 +446,10 @@ namespace Test
 	}
 
 	bool RunTestInferenceAndCompareToRef(const FTests::FTestSetup& TestSetup, IRuntime* Runtime, const FNNIModelRaw& ONNXModel, 
-		TArrayView<TArray<char>> RefOutputMemBuffers, TArrayView<FTensor> RefOutputTensors)
+		TArrayView<TArray<char>> RefOutputMemBuffers, TArrayView<UE::NNECore::Internal::FTensor> RefOutputTensors)
 	{
 		TArray<TArray<char>> OutputMemBuffers;
-		TArray<FTensor> OutputTensors;
+		TArray<UE::NNECore::Internal::FTensor> OutputTensors;
 
 		const FString& RuntimeName = Runtime->GetRuntimeName();
 		const float AbsoluteErrorEpsilon = TestSetup.GetAbsoluteErrorEpsilonForRuntime(RuntimeName);
@@ -501,7 +501,7 @@ namespace Test
 		const float AbsoluteRefErrorEpsilon = TestSetup.GetAbsoluteErrorEpsilonForRuntime(RefName);
 		const float RelativeRefErrorPercent = TestSetup.GetRelativeErrorPercentForRuntime(RefName);
 		TArray<TArray<char>> RefOutputMemBuffers;
-		TArray<FTensor> RefOutputTensors;
+		TArray<UE::NNECore::Internal::FTensor> RefOutputTensors;
 		bool bAllTestsSucceeded = true;
 
 		if (!RunTestInference(ONNXModel, TestSetup, RefRuntime, RefOutputTensors, RefOutputMemBuffers))
diff --git a/Engine/Plugins/Experimental/NNE/Source/NNXQA/Private/NNXQAUtils.h b/Engine/Plugins/Experimental/NNE/Source/NNXQA/Private/NNXQAUtils.h
--- a/Engine/Plugins/Experimental/NNE/Source/NNXQA/Private/NNXQAUtils.h
+++ b/Engine/Plugins/Experimental/NNE/Source/NNXQA/Private/NNXQAUtils.h
@@ -7,11 +7,13 @@
 #include "NNECoreAttributeMap.h"
 #include "NNXRuntime.h"
 #include "NNXRuntimeFormat.h"
-#include "NNXTypes.h"
+#include "NNECoreTypes.h"
+#include "NNECoreTensor.h"
 
 
 namespace NNX 
 {
+
 namespace Test 
 {
 	struct FTests
@@ -38,9 +40,9 @@ namespace Test
 			bool IsModelTest;
 			TMap<FString, float> AbsoluteErrorEpsilonForRuntime;
 			TMap<FString, float> RelativeErrorPercentForRuntime;
-			TArray<FTensor> Inputs;
-			TArray<FTensor> Weights;
-			TArray<FTensor> Outputs;
+			TArray<UE::NNECore::Internal::FTensor> Inputs;
+			TArray<UE::NNECore::Internal::FTensor> Weights;
+			TArray<UE::NNECore::Internal::FTensor> Outputs;
 			TArray<FTensorData> InputsData;
 			TArray<FTensorData> WeightsData;
 			TArray<FTensorData> OutputsData;
@@ -74,18 +76,18 @@ namespace Test
 	
 	class ElementWiseCosTensorInitializer
 	{
-		EMLTensorDataType DataType;
+		ENNETensorDataType DataType;
 		uint32 TensorIndex;
 
 	public:
-		ElementWiseCosTensorInitializer(EMLTensorDataType InDataType, uint32 InTensorIndex);
+		ElementWiseCosTensorInitializer(ENNETensorDataType InDataType, uint32 InTensorIndex);
 		float operator () (uint32 ElementIndex) const;
 	};
 
-	FString TensorToString(const FTensor& Tensor);
-	FString TensorToString(const FTensor& Tensor, TConstArrayView<char> TensorData);
+	FString TensorToString(const UE::NNECore::Internal::FTensor& Tensor);
+	FString TensorToString(const UE::NNECore::Internal::FTensor& Tensor, TConstArrayView<char> TensorData);
 	template<typename T> FString ShapeToString(TConstArrayView<T> Shape);
-	TArray<char> GenerateTensorDataForTest(const FTensor& Tensor, std::function<float(uint32)> ElementInitializer);
+	TArray<char> GenerateTensorDataForTest(const UE::NNECore::Internal::FTensor& Tensor, std::function<float(uint32)> ElementInitializer);
 
 } // namespace Test
 } // namespace NNX
\ No newline at end of file
diff --git a/Engine/Plugins/Experimental/NNE/Source/NNXRuntimeCPU/Private/NNXRuntimeCPU.cpp b/Engine/Plugins/Experimental/NNE/Source/NNXRuntimeCPU/Private/NNXRuntimeCPU.cpp
--- a/Engine/Plugins/Experimental/NNE/Source/NNXRuntimeCPU/Private/NNXRuntimeCPU.cpp
+++ b/Engine/Plugins/Experimental/NNE/Source/NNXRuntimeCPU/Private/NNXRuntimeCPU.cpp
@@ -233,15 +233,16 @@ bool FMLInferenceModelCPU::ConfigureTensors(const bool InIsInput)
 
 		TensorsORTType.Emplace(ONNXTensorElementDataTypeEnum);
 
-		std::pair<EMLTensorDataType, uint64> TypeAndSize = TranslateTensorTypeORTToNNI(ONNXTensorElementDataTypeEnum);
+		std::pair<ENNETensorDataType, uint64> TypeAndSize = TranslateTensorTypeORTToNNI(ONNXTensorElementDataTypeEnum);
 
-		FSymbolicTensorShape Shape;
-		Shape.Data.Reserve(CurrentTensorInfo.GetShape().size());
+		TArray<int32> ShapeData;
+		ShapeData.Reserve(CurrentTensorInfo.GetShape().size());
 		for (const int64_t CurrentTensorSize : CurrentTensorInfo.GetShape())
 		{
-			Shape.Data.Add((int32)CurrentTensorSize);
+			ShapeData.Add((int32)CurrentTensorSize);
 		}
 
+		UE::NNECore::FSymbolicTensorShape Shape = UE::NNECore::FSymbolicTensorShape::Make(ShapeData);
 		FTensorDesc SymbolicTensorDesc = FTensorDesc::Make(FString(CurTensorName), Shape, TypeAndSize.first);
 		
 		check(SymbolicTensorDesc.GetElemByteSize() == TypeAndSize.second);
@@ -266,7 +267,7 @@ int FMLInferenceModelCPU::SetInputTensorShapes(TConstArrayView<FTensorShape> InI
 	// Setup concrete input tensor
 	for (int i = 0; i < InputSymbolicTensors.Num(); ++i)
 	{
-		FTensor Tensor = FTensor::Make(InputSymbolicTensors[i].GetName(), InInputShapes[i], InputSymbolicTensors[i].GetDataType());
+		UE::NNECore::Internal::FTensor Tensor = UE::NNECore::Internal::FTensor::Make(InputSymbolicTensors[i].GetName(), InInputShapes[i], InputSymbolicTensors[i].GetDataType());
 		InputTensors.Emplace(Tensor);
 	}
 
@@ -276,9 +277,9 @@ int FMLInferenceModelCPU::SetInputTensorShapes(TConstArrayView<FTensorShape> InI
 	// Setup concrete output shapes only if all model output shapes are concretes, otherwise it will be set during Run()
 	for (FTensorDesc SymbolicTensorDesc : OutputSymbolicTensors)
 	{
-		if (SymbolicTensorDesc.IsConcrete())
+		if (SymbolicTensorDesc.GetShape().IsConcrete())
 		{
-			FTensor Tensor = FTensor::MakeFromSymbolicDesc(SymbolicTensorDesc);
+			UE::NNECore::Internal::FTensor Tensor = UE::NNECore::Internal::FTensor::MakeFromSymbolicDesc(SymbolicTensorDesc);
 			OutputTensors.Emplace(Tensor);
 			OutputTensorShapes.Emplace(Tensor.GetShape());
 		}
diff --git a/Engine/Plugins/Experimental/NNE/Source/NNXRuntimeCPU/Public/NNXRuntimeCPU.h b/Engine/Plugins/Experimental/NNE/Source/NNXRuntimeCPU/Public/NNXRuntimeCPU.h
--- a/Engine/Plugins/Experimental/NNE/Source/NNXRuntimeCPU/Public/NNXRuntimeCPU.h
+++ b/Engine/Plugins/Experimental/NNE/Source/NNXRuntimeCPU/Public/NNXRuntimeCPU.h
@@ -4,9 +4,11 @@
 
 #include "CoreMinimal.h"
 #include "NNXCore.h"
+#include "NNECoreTensor.h"
+#include "NNECoreTypes.h"
+#include "NNEProfilingStatistics.h"
 #include "NNXRuntime.h"
 #include "NNXInferenceModel.h"
-#include "NNEProfilingStatistics.h"
 #include "NNXModelOptimizerInterface.h"
 
 #include "NNXThirdPartyWarningDisabler.h"
@@ -117,8 +119,8 @@ namespace NNX
 		TArray<const char*> InputTensorNames;
 		TArray<const char*> OutputTensorNames;
 
-		TArray<FTensor> InputTensors;
-		TArray<FTensor> OutputTensors;
+		TArray<UE::NNECore::Internal::FTensor> InputTensors;
+		TArray<UE::NNECore::Internal::FTensor> OutputTensors;
 
 		bool InitializedAndConfigureMembers();
 		bool ConfigureTensors(const bool InIsInput);
diff --git a/Engine/Plugins/Experimental/NNE/Source/NNXRuntimeCPU/Public/NNXRuntimeCPUUtils.h b/Engine/Plugins/Experimental/NNE/Source/NNXRuntimeCPU/Public/NNXRuntimeCPUUtils.h
--- a/Engine/Plugins/Experimental/NNE/Source/NNXRuntimeCPU/Public/NNXRuntimeCPUUtils.h
+++ b/Engine/Plugins/Experimental/NNE/Source/NNXRuntimeCPU/Public/NNXRuntimeCPUUtils.h
@@ -3,7 +3,8 @@
 #pragma once
 
 #include "HAL/UnrealMemory.h"
-#include "NNXTypes.h"
+#include "NNECoreTypes.h"
+#include "NNECoreTensor.h"
 #include "NNXRuntime.h"
 #include "NNXThirdPartyWarningDisabler.h"
 NNX_THIRD_PARTY_INCLUDES_START
@@ -14,98 +15,98 @@ NNX_THIRD_PARTY_INCLUDES_END
 
 DECLARE_STATS_GROUP(TEXT("MachineLearning"), STATGROUP_MachineLearning, STATCAT_Advanced);
 
-using TypeInfoORT = std::pair<EMLTensorDataType, uint64>;
+using TypeInfoORT = std::pair<ENNETensorDataType, uint64>;
 
 inline TypeInfoORT TranslateTensorTypeORTToNNI(unsigned int OrtDataType) 
 {
-	EMLTensorDataType DataType = EMLTensorDataType::None;
+	ENNETensorDataType DataType = ENNETensorDataType::None;
 	uint64 ElementSize = 0;
 
 	switch (OrtDataType) {
 	case ONNX_TENSOR_ELEMENT_DATA_TYPE_UNDEFINED:
-		DataType = EMLTensorDataType::None;
+		DataType = ENNETensorDataType::None;
 		ElementSize = 0;
 		break;
 
 	case ONNX_TENSOR_ELEMENT_DATA_TYPE_FLOAT:
-		DataType = EMLTensorDataType::Float;
+		DataType = ENNETensorDataType::Float;
 		ElementSize = sizeof(float);
 		break;
 
 	case ONNX_TENSOR_ELEMENT_DATA_TYPE_UINT8:
-		DataType = EMLTensorDataType::UInt8;
+		DataType = ENNETensorDataType::UInt8;
 		ElementSize = sizeof(uint8);
 		break;
 	case ONNX_TENSOR_ELEMENT_DATA_TYPE_INT8:
-		DataType = EMLTensorDataType::Int8;
+		DataType = ENNETensorDataType::Int8;
 		ElementSize = sizeof(int8);
 		break;
 	case ONNX_TENSOR_ELEMENT_DATA_TYPE_UINT16:
-		DataType = EMLTensorDataType::UInt16;
+		DataType = ENNETensorDataType::UInt16;
 		ElementSize = sizeof(uint16);
 		break;
 	case ONNX_TENSOR_ELEMENT_DATA_TYPE_INT16:
-		DataType = EMLTensorDataType::Int16;
+		DataType = ENNETensorDataType::Int16;
 		ElementSize = sizeof(int16);
 		break;
 
 	case ONNX_TENSOR_ELEMENT_DATA_TYPE_INT32:
-		DataType = EMLTensorDataType::Int32;
+		DataType = ENNETensorDataType::Int32;
 		ElementSize = sizeof(int32);
 		break;
 
 	case ONNX_TENSOR_ELEMENT_DATA_TYPE_INT64:
-		DataType = EMLTensorDataType::Int64;
+		DataType = ENNETensorDataType::Int64;
 		ElementSize = sizeof(int64);
 		break;
 
 	case ONNX_TENSOR_ELEMENT_DATA_TYPE_STRING:
-		DataType = EMLTensorDataType::Char;
+		DataType = ENNETensorDataType::Char;
 		ElementSize = sizeof(char);
 		break;
 
 	case ONNX_TENSOR_ELEMENT_DATA_TYPE_BOOL:
-		DataType = EMLTensorDataType::Boolean;
+		DataType = ENNETensorDataType::Boolean;
 		ElementSize = sizeof(bool);
 		break;
 
 	case ONNX_TENSOR_ELEMENT_DATA_TYPE_FLOAT16:
-		DataType = EMLTensorDataType::Half;
+		DataType = ENNETensorDataType::Half;
 		ElementSize = 2;
 		break;
 
 	case ONNX_TENSOR_ELEMENT_DATA_TYPE_DOUBLE:
-		DataType = EMLTensorDataType::Double;
+		DataType = ENNETensorDataType::Double;
 		ElementSize = sizeof(double);
 		break;
 
 	case ONNX_TENSOR_ELEMENT_DATA_TYPE_UINT32:
-		DataType = EMLTensorDataType::UInt32;
+		DataType = ENNETensorDataType::UInt32;
 		ElementSize = sizeof(uint32);
 		break;
 
 	case ONNX_TENSOR_ELEMENT_DATA_TYPE_UINT64:
-		DataType = EMLTensorDataType::UInt64;
+		DataType = ENNETensorDataType::UInt64;
 		ElementSize = sizeof(uint64);
 		break;
 
 	case ONNX_TENSOR_ELEMENT_DATA_TYPE_COMPLEX64:
-		DataType = EMLTensorDataType::Complex64;
+		DataType = ENNETensorDataType::Complex64;
 		ElementSize = 8;
 		break;
 
 	case ONNX_TENSOR_ELEMENT_DATA_TYPE_COMPLEX128:
-		DataType = EMLTensorDataType::Complex128;
+		DataType = ENNETensorDataType::Complex128;
 		ElementSize = 16;
 		break;
 
 	case ONNX_TENSOR_ELEMENT_DATA_TYPE_BFLOAT16:
-		DataType = EMLTensorDataType::BFloat16;
+		DataType = ENNETensorDataType::BFloat16;
 		ElementSize = 2;
 		break;
 
 	default:
-		DataType = EMLTensorDataType::None;
+		DataType = ENNETensorDataType::None;
 		break;
 	}
 
@@ -114,7 +115,7 @@ inline TypeInfoORT TranslateTensorTypeORTToNNI(unsigned int OrtDataType)
 
 inline void BindTensorsToORT(
 	TConstArrayView<const NNX::FMLTensorBinding> InBindingTensors,
-	TConstArrayView<NNX::FTensor> InTensors,
+	TConstArrayView<UE::NNECore::Internal::FTensor> InTensors,
 	TConstArrayView<ONNXTensorElementDataType> InTensorsORTType, 
 	const Ort::MemoryInfo* InAllocatorInfo,
 	TArray<Ort::Value>& OutOrtTensors
@@ -132,14 +133,14 @@ inline void BindTensorsToORT(
 	for (uint32 Index = 0; Index < NumBinding; ++Index)
 	{
 		const NNX::FMLTensorBinding& Binding = InBindingTensors[Index];
-		const NNX::FTensor& Tensor = InTensors[Index];
+		const UE::NNECore::Internal::FTensor& Tensor = InTensors[Index];
 		const ONNXTensorElementDataType CurrentORTType = InTensorsORTType[Index];
 
 		TUniquePtr<int64_t[]> SizesInt64t;
 		SizesInt64t = MakeUnique<int64_t[]>(Tensor.GetShape().Rank());
 		for (int32 DimIndex = 0; DimIndex < Tensor.GetShape().Rank(); ++DimIndex)
 		{
-			SizesInt64t.Get()[DimIndex] = Tensor.GetShape().Data[DimIndex];
+			SizesInt64t.Get()[DimIndex] = Tensor.GetShape().GetData()[DimIndex];
 		}
 
 		const uint64 ByteCount { InTensors[Index].GetDataSize() };
@@ -160,8 +161,8 @@ inline void BindTensorsToORT(
 inline void CopyFromORTToBindings(
 	TConstArrayView<Ort::Value> InOrtTensors,
 	TConstArrayView<const NNX::FMLTensorBinding> InBindingTensors,
-	TConstArrayView<NNX::FTensorDesc> InTensorDescs,
-	TArray<NNX::FTensor>& OutTensors
+	TConstArrayView<UE::NNECore::FTensorDesc> InTensorDescs,
+	TArray<UE::NNECore::Internal::FTensor>& OutTensors
 )
 {
 	const uint32 NumBinding = InOrtTensors.Num();
@@ -176,18 +177,19 @@ inline void CopyFromORTToBindings(
 	for (uint32 Index = 0; Index < NumBinding; ++Index)
 	{
 		const NNX::FMLTensorBinding& Binding = InBindingTensors[Index];
-		const NNX::FTensorDesc& TensorDesc = InTensorDescs[Index];
+		const UE::NNECore::FTensorDesc& TensorDesc = InTensorDescs[Index];
 		const Ort::Value& OrtTensor = InOrtTensors[Index];
 		const std::vector<int64_t>& OrtShape = OrtTensor.GetTensorTypeAndShapeInfo().GetShape();
 
-		NNX::FTensorShape Shape;
+		TArray<uint32> ShapeData;
 		for (int32 DimIndex = 0; DimIndex < OrtShape.size(); ++DimIndex)
 		{
 			check(OrtShape[DimIndex] >= 0);
-			Shape.Data.Add(OrtShape[DimIndex]);
+			ShapeData.Add(OrtShape[DimIndex]);
 		}
 
-		NNX::FTensor Tensor = NNX::FTensor::Make(TensorDesc.GetName(), Shape, TensorDesc.GetDataType());
+		UE::NNECore::FTensorShape Shape = UE::NNECore::FTensorShape::Make(ShapeData);
+		UE::NNECore::Internal::FTensor Tensor = UE::NNECore::Internal::FTensor::Make(TensorDesc.GetName(), Shape, TensorDesc.GetDataType());
 		OutTensors.Emplace(Tensor);
 
 		void* CpuMemory = (uint8_t*)Binding.CpuMemory + Binding.OffsetInBytes;
diff --git a/Engine/Plugins/Experimental/NNE/Source/NNXRuntimeORT/Private/NNXRuntimeORT.cpp b/Engine/Plugins/Experimental/NNE/Source/NNXRuntimeORT/Private/NNXRuntimeORT.cpp
--- a/Engine/Plugins/Experimental/NNE/Source/NNXRuntimeORT/Private/NNXRuntimeORT.cpp
+++ b/Engine/Plugins/Experimental/NNE/Source/NNXRuntimeORT/Private/NNXRuntimeORT.cpp
@@ -287,16 +287,18 @@ bool FMLInferenceModelORT::ConfigureTensors(const bool InIsInput)
 
 		TensorsORTType.Emplace(ONNXTensorElementDataTypeEnum);
 
-		std::pair<EMLTensorDataType, uint64> TypeAndSize = TranslateTensorTypeORTToNNI(ONNXTensorElementDataTypeEnum);
+		std::pair<ENNETensorDataType, uint64> TypeAndSize = TranslateTensorTypeORTToNNI(ONNXTensorElementDataTypeEnum);
 		
-		FSymbolicTensorShape Shape;
-		Shape.Data.Reserve(CurrentTensorInfo.GetShape().size());
+		
+		TArray<int32> ShapeData;
+		ShapeData.Reserve(CurrentTensorInfo.GetShape().size());
 		for (const int64_t CurrentTensorSize : CurrentTensorInfo.GetShape())
 		{
-			Shape.Data.Add((int32)CurrentTensorSize);
+			ShapeData.Add((int32)CurrentTensorSize);
 		}
 
-		FTensorDesc SymbolicTensorDesc = FTensorDesc::Make(FString(CurTensorName), Shape, TypeAndSize.first);
+		UE::NNECore::FSymbolicTensorShape Shape = UE::NNECore::FSymbolicTensorShape::Make(ShapeData);
+		UE::NNECore::FTensorDesc SymbolicTensorDesc = UE::NNECore::FTensorDesc::Make(FString(CurTensorName), Shape, TypeAndSize.first);
 		
 		check(SymbolicTensorDesc.GetElemByteSize() == TypeAndSize.second);
 		SymbolicTensorDescs.Emplace(SymbolicTensorDesc);
@@ -320,7 +322,7 @@ int FMLInferenceModelORT::SetInputTensorShapes(TConstArrayView<FTensorShape> InI
 	// Setup concrete input tensor
 	for (int i = 0; i < InputSymbolicTensors.Num(); ++i)
 	{
-		FTensor Tensor = FTensor::Make(InputSymbolicTensors[i].GetName(), InInputShapes[i], InputSymbolicTensors[i].GetDataType());
+		UE::NNECore::Internal::FTensor Tensor = UE::NNECore::Internal::FTensor::Make(InputSymbolicTensors[i].GetName(), InInputShapes[i], InputSymbolicTensors[i].GetDataType());
 		InputTensors.Emplace(Tensor);
 	}
 
@@ -330,9 +332,9 @@ int FMLInferenceModelORT::SetInputTensorShapes(TConstArrayView<FTensorShape> InI
 	// Setup concrete output shapes only if all model output shapes are concretes, otherwise it will be set during Run()
 	for (FTensorDesc SymbolicTensorDesc : OutputSymbolicTensors)
 	{
-		if (SymbolicTensorDesc.IsConcrete())
+		if (SymbolicTensorDesc.GetShape().IsConcrete())
 		{
-			FTensor Tensor = FTensor::MakeFromSymbolicDesc(SymbolicTensorDesc);
+			UE::NNECore::Internal::FTensor Tensor = UE::NNECore::Internal::FTensor::MakeFromSymbolicDesc(SymbolicTensorDesc);
 			OutputTensors.Emplace(Tensor);
 			OutputTensorShapes.Emplace(Tensor.GetShape());
 		}
diff --git a/Engine/Plugins/Experimental/NNE/Source/NNXRuntimeORT/Private/NNXRuntimeORT.h b/Engine/Plugins/Experimental/NNE/Source/NNXRuntimeORT/Private/NNXRuntimeORT.h
--- a/Engine/Plugins/Experimental/NNE/Source/NNXRuntimeORT/Private/NNXRuntimeORT.h
+++ b/Engine/Plugins/Experimental/NNE/Source/NNXRuntimeORT/Private/NNXRuntimeORT.h
@@ -4,6 +4,8 @@
 
 #include "CoreMinimal.h"
 #include "NNXCore.h"
+#include "NNECoreTensor.h"
+#include "NNECoreTypes.h"
 #include "NNXRuntime.h"
 #include "NNXInferenceModel.h"
 #include "NNXModelOptimizerInterface.h"
@@ -217,8 +219,8 @@ namespace NNX
 		TArray<const char*> InputTensorNames;
 		TArray<const char*> OutputTensorNames;
 
-		TArray<FTensor> InputTensors;
-		TArray<FTensor> OutputTensors;
+		TArray<UE::NNECore::Internal::FTensor> InputTensors;
+		TArray<UE::NNECore::Internal::FTensor> OutputTensors;
 
 		virtual bool InitializedAndConfigureMembers();
 		bool ConfigureTensors(const bool InIsInput);
diff --git a/Engine/Plugins/Experimental/NNE/Source/NNXRuntimeORT/Private/NNXRuntimeORTUtils.h b/Engine/Plugins/Experimental/NNE/Source/NNXRuntimeORT/Private/NNXRuntimeORTUtils.h
--- a/Engine/Plugins/Experimental/NNE/Source/NNXRuntimeORT/Private/NNXRuntimeORTUtils.h
+++ b/Engine/Plugins/Experimental/NNE/Source/NNXRuntimeORT/Private/NNXRuntimeORTUtils.h
@@ -3,7 +3,8 @@
 #pragma once
 
 #include "HAL/UnrealMemory.h"
-#include "NNXTypes.h"
+#include "NNECoreTensor.h"
+#include "NNECoreTypes.h"
 #include "NNXRuntime.h"
 #include "NNXThirdPartyWarningDisabler.h"
 NNX_THIRD_PARTY_INCLUDES_START
@@ -14,98 +15,98 @@ NNX_THIRD_PARTY_INCLUDES_END
 
 DECLARE_STATS_GROUP(TEXT("MachineLearning"), STATGROUP_MachineLearning, STATCAT_Advanced);
 
-using TypeInfoORT = std::pair<EMLTensorDataType, uint64>;
+using TypeInfoORT = std::pair<ENNETensorDataType, uint64>;
 
 inline TypeInfoORT TranslateTensorTypeORTToNNI(unsigned int OrtDataType)
 {
-	EMLTensorDataType DataType = EMLTensorDataType::None;
+	ENNETensorDataType DataType = ENNETensorDataType::None;
 	uint64 ElementSize = 0;
 
 	switch (OrtDataType) {
 	case ONNX_TENSOR_ELEMENT_DATA_TYPE_UNDEFINED:
-		DataType = EMLTensorDataType::None;
+		DataType = ENNETensorDataType::None;
 		ElementSize = 0;
 		break;
 
 	case ONNX_TENSOR_ELEMENT_DATA_TYPE_FLOAT:
-		DataType = EMLTensorDataType::Float;
+		DataType = ENNETensorDataType::Float;
 		ElementSize = sizeof(float);
 		break;
 
 	case ONNX_TENSOR_ELEMENT_DATA_TYPE_UINT8:
-		DataType = EMLTensorDataType::UInt8;
+		DataType = ENNETensorDataType::UInt8;
 		ElementSize = sizeof(uint8);
 		break;
 	case ONNX_TENSOR_ELEMENT_DATA_TYPE_INT8:
-		DataType = EMLTensorDataType::Int8;
+		DataType = ENNETensorDataType::Int8;
 		ElementSize = sizeof(int8);
 		break;
 	case ONNX_TENSOR_ELEMENT_DATA_TYPE_UINT16:
-		DataType = EMLTensorDataType::UInt16;
+		DataType = ENNETensorDataType::UInt16;
 		ElementSize = sizeof(uint16);
 		break;
 	case ONNX_TENSOR_ELEMENT_DATA_TYPE_INT16:
-		DataType = EMLTensorDataType::Int16;
+		DataType = ENNETensorDataType::Int16;
 		ElementSize = sizeof(int16);
 		break;
 
 	case ONNX_TENSOR_ELEMENT_DATA_TYPE_INT32:
-		DataType = EMLTensorDataType::Int32;
+		DataType = ENNETensorDataType::Int32;
 		ElementSize = sizeof(int32);
 		break;
 
 	case ONNX_TENSOR_ELEMENT_DATA_TYPE_INT64:
-		DataType = EMLTensorDataType::Int64;
+		DataType = ENNETensorDataType::Int64;
 		ElementSize = sizeof(int64);
 		break;
 
 	case ONNX_TENSOR_ELEMENT_DATA_TYPE_STRING:
-		DataType = EMLTensorDataType::Char;
+		DataType = ENNETensorDataType::Char;
 		ElementSize = sizeof(char);
 		break;
 
 	case ONNX_TENSOR_ELEMENT_DATA_TYPE_BOOL:
-		DataType = EMLTensorDataType::Boolean;
+		DataType = ENNETensorDataType::Boolean;
 		ElementSize = sizeof(bool);
 		break;
 
 	case ONNX_TENSOR_ELEMENT_DATA_TYPE_FLOAT16:
-		DataType = EMLTensorDataType::Half;
+		DataType = ENNETensorDataType::Half;
 		ElementSize = 2;
 		break;
 
 	case ONNX_TENSOR_ELEMENT_DATA_TYPE_DOUBLE:
-		DataType = EMLTensorDataType::Double;
+		DataType = ENNETensorDataType::Double;
 		ElementSize = sizeof(double);
 		break;
 
 	case ONNX_TENSOR_ELEMENT_DATA_TYPE_UINT32:
-		DataType = EMLTensorDataType::UInt32;
+		DataType = ENNETensorDataType::UInt32;
 		ElementSize = sizeof(uint32);
 		break;
 
 	case ONNX_TENSOR_ELEMENT_DATA_TYPE_UINT64:
-		DataType = EMLTensorDataType::UInt64;
+		DataType = ENNETensorDataType::UInt64;
 		ElementSize = sizeof(uint64);
 		break;
 
 	case ONNX_TENSOR_ELEMENT_DATA_TYPE_COMPLEX64:
-		DataType = EMLTensorDataType::Complex64;
+		DataType = ENNETensorDataType::Complex64;
 		ElementSize = 8;
 		break;
 
 	case ONNX_TENSOR_ELEMENT_DATA_TYPE_COMPLEX128:
-		DataType = EMLTensorDataType::Complex128;
+		DataType = ENNETensorDataType::Complex128;
 		ElementSize = 16;
 		break;
 
 	case ONNX_TENSOR_ELEMENT_DATA_TYPE_BFLOAT16:
-		DataType = EMLTensorDataType::BFloat16;
+		DataType = ENNETensorDataType::BFloat16;
 		ElementSize = 2;
 		break;
 
 	default:
-		DataType = EMLTensorDataType::None;
+		DataType = ENNETensorDataType::None;
 		break;
 	}
 
@@ -114,7 +115,7 @@ inline TypeInfoORT TranslateTensorTypeORTToNNI(unsigned int OrtDataType)
 
 inline void BindTensorsToORT(
 	TConstArrayView<const NNX::FMLTensorBinding> InBindingTensors,
-	TConstArrayView<NNX::FTensor> InTensors,
+	TConstArrayView<UE::NNECore::Internal::FTensor> InTensors,
 	TConstArrayView<ONNXTensorElementDataType> InTensorsORTType,
 	const Ort::MemoryInfo* InAllocatorInfo,
 	TArray<Ort::Value>& OutOrtTensors
@@ -132,14 +133,14 @@ inline void BindTensorsToORT(
 	for (uint32 Index = 0; Index < NumBinding; ++Index)
 	{
 		const NNX::FMLTensorBinding& Binding = InBindingTensors[Index];
-		const NNX::FTensor& Tensor = InTensors[Index];
+		const UE::NNECore::Internal::FTensor& Tensor = InTensors[Index];
 		const ONNXTensorElementDataType CurrentORTType = InTensorsORTType[Index];
 
 		TUniquePtr<int64_t[]> SizesInt64t;
 		SizesInt64t = MakeUnique<int64_t[]>(Tensor.GetShape().Rank());
 		for (int32 DimIndex = 0; DimIndex < Tensor.GetShape().Rank(); ++DimIndex)
 		{
-			SizesInt64t.Get()[DimIndex] = Tensor.GetShape().Data[DimIndex];
+			SizesInt64t.Get()[DimIndex] = Tensor.GetShape().GetData()[DimIndex];
 		}
 
 		const uint64 ByteCount{ InTensors[Index].GetDataSize() };
@@ -160,8 +161,8 @@ inline void BindTensorsToORT(
 inline void CopyFromORTToBindings(
 	TConstArrayView<Ort::Value> InOrtTensors,
 	TConstArrayView<const NNX::FMLTensorBinding> InBindingTensors,
-	TConstArrayView<NNX::FTensorDesc> InTensorDescs,
-	TArray<NNX::FTensor>& OutTensors
+	TConstArrayView<UE::NNECore::FTensorDesc> InTensorDescs,
+	TArray<UE::NNECore::Internal::FTensor>& OutTensors
 )
 {
 	const uint32 NumBinding = InOrtTensors.Num();
@@ -176,18 +177,19 @@ inline void CopyFromORTToBindings(
 	for (uint32 Index = 0; Index < NumBinding; ++Index)
 	{
 		const NNX::FMLTensorBinding& Binding = InBindingTensors[Index];
-		const NNX::FTensorDesc& TensorDesc = InTensorDescs[Index];
+		const UE::NNECore::FTensorDesc& TensorDesc = InTensorDescs[Index];
 		const Ort::Value& OrtTensor = InOrtTensors[Index];
 		const std::vector<int64_t>& OrtShape = OrtTensor.GetTensorTypeAndShapeInfo().GetShape();
 
-		NNX::FTensorShape Shape;
+		TArray<uint32> ShapeData;
 		for (int32 DimIndex = 0; DimIndex < OrtShape.size(); ++DimIndex)
 		{
 			check(OrtShape[DimIndex] >= 0);
-			Shape.Data.Add(OrtShape[DimIndex]);
+			ShapeData.Add(OrtShape[DimIndex]);
 		}
 
-		NNX::FTensor Tensor = NNX::FTensor::Make(TensorDesc.GetName(), Shape, TensorDesc.GetDataType());
+		UE::NNECore::FTensorShape Shape = UE::NNECore::FTensorShape::Make(ShapeData);
+		UE::NNECore::Internal::FTensor Tensor = UE::NNECore::Internal::FTensor::Make(TensorDesc.GetName(), Shape, TensorDesc.GetDataType());
 		OutTensors.Emplace(Tensor);
 
 		void* CpuMemory = (uint8_t*)Binding.CpuMemory + Binding.OffsetInBytes;
diff --git a/Engine/Plugins/Experimental/NNE/Source/NNXRuntimeRDG/Private/Hlsl/NNIRuntimeRDGConv.cpp b/Engine/Plugins/Experimental/NNE/Source/NNXRuntimeRDG/Private/Hlsl/NNIRuntimeRDGConv.cpp
--- a/Engine/Plugins/Experimental/NNE/Source/NNXRuntimeRDG/Private/Hlsl/NNIRuntimeRDGConv.cpp
+++ b/Engine/Plugins/Experimental/NNE/Source/NNXRuntimeRDG/Private/Hlsl/NNIRuntimeRDGConv.cpp
@@ -4,10 +4,12 @@
 #include "NNEHlslShadersConvCS.h"
 #include "NNECoreAttributeMap.h"
 #include "NNXRuntimeHLSLHelper.h"
+#include "NNECoreTensor.h"
+#include "NNECoreTypes.h"
 
 namespace UE::NNIRuntimeRDG::Private::Hlsl
 {
-	DECLARE_GPU_STAT_NAMED(FNNIOperatorConv, TEXT("NNI.Operator.Hlsl.Conv"));
+	DECLARE_GPU_STAT_NAMED(FNNEOperatorConv, TEXT("NNE.Operator.Hlsl.Conv"));
 
 	using EConvAutoPad = UE::NNEHlslShaders::Internal::EConvAutoPad;
 
@@ -39,35 +41,36 @@ namespace UE::NNIRuntimeRDG::Private::Hlsl
 
 	public:
 
-		virtual int PrepareOutputs(TConstArrayView<NNX::FTensorRef> InputTensors, TArrayView<NNX::FTensorRef> OutputTensors) const override
+		virtual int PrepareOutputs(TConstArrayView<NNECore::Internal::FTensorRef> InputTensors, TArrayView<NNECore::Internal::FTensorRef> OutputTensors) const override
 		{
 			check(InputTensors.Num() >= 2 && InputTensors.Num() <= 3);
 			check(OutputTensors.Num() == 1);
 
-			const NNX::FTensorShape& Input = InputTensors[0]->GetShape();
-			const NNX::FTensorShape& Weights = InputTensors[1]->GetShape();
-			NNX::FSymbolicTensorShape OutputShape;
+			const NNECore::FTensorShape& Input = InputTensors[0]->GetShape();
+			const NNECore::FTensorShape& Weights = InputTensors[1]->GetShape();
+
+			const TArray<int32> OutputShapeData = NNEHlslShaders::Internal::FConvCS::GetOutputShape(Input.GetData(), Weights.GetData(), AutoPad, Dilations, Strides, Pads);
+			const NNECore::FSymbolicTensorShape OutputShape = NNECore::FSymbolicTensorShape::Make(OutputShapeData);
 			
-			OutputShape.Data = UE::NNEHlslShaders::Internal::FConvCS::GetOutputShape(Input.Data, Weights.Data, AutoPad, Dilations, Strides, Pads);
 			if (!OutputShape.IsConcrete())
 			{
 				return -1;
 			}
-			OutputTensors[0]->SetShape(NNX::FTensorShape::MakeFromSymbolic(OutputShape));
+			OutputTensors[0]->SetShape(NNECore::FTensorShape::MakeFromSymbolic(OutputShape));
 
 			return 0;
 		};
 
-		virtual bool Initialize(TConstArrayView<NNX::FTensorDesc> InputTensorDescs, TConstArrayView<NNX::FTensorDesc> OutputTensorDescs, const UE::NNECore::FAttributeMap& Attributes) override
+		virtual bool Initialize(TConstArrayView<NNECore::FTensorDesc> InputTensorDescs, TConstArrayView<NNECore::FTensorDesc> OutputTensorDescs, const NNECore::FAttributeMap& Attributes) override
 		{
             using namespace UE::NNEHlslShaders::Internal;
 
 			check(InputTensorDescs.Num() >= 2 && InputTensorDescs.Num() <= 3);
 			check(OutputTensorDescs.Num() == 1);
 
-            const NNX::FTensorDesc& Input = InputTensorDescs[0];
-			const NNX::FTensorDesc& Weights = InputTensorDescs[1];
-			const NNX::FTensorDesc& Output = OutputTensorDescs[0];
+            const NNECore::FTensorDesc& Input = InputTensorDescs[0];
+			const NNECore::FTensorDesc& Weights = InputTensorDescs[1];
+			const NNECore::FTensorDesc& Output = OutputTensorDescs[0];
 			
 			if (Input.GetShape().Rank() < 2)
 			{
@@ -135,11 +138,11 @@ namespace UE::NNIRuntimeRDG::Private::Hlsl
 			check(Output.GetShape().Rank() == Input.GetShape().Rank());
 			check(NumDimensions == (Input.GetShape().Rank() - 2));
 
-			TArray<int32> OutputShape = FConvCS::GetOutputShape(Input.GetShape().Data, Weights.GetShape().Data, AutoPad, Dilations, Strides, Pads);
+			TArray<int32> OutputShape = FConvCS::GetOutputShape(Input.GetShape().GetData(), Weights.GetShape().GetData(), AutoPad, Dilations, Strides, Pads);
 
 			// Set parameters
 			FConvCS::FParameters* Params = GraphBuilder.AllocParameters<FConvCS::FParameters>();
-			FConvCS::FillInParameters(GroupSize, Input.GetShape().Data, Weights.GetShape().Data, HasBias, AutoPad, Group, Dilations,Strides, Pads, *Params);
+			FConvCS::FillInParameters(GroupSize, Input.GetShape().GetData(), Weights.GetShape().GetData(), HasBias, AutoPad, Group, Dilations,Strides, Pads, *Params);
 			Params->X = GraphBuilder.CreateSRV(FRDGBufferSRVDesc(Input.GetBuffer(), PF_R32_FLOAT));
 			Params->W = GraphBuilder.CreateSRV(FRDGBufferSRVDesc(Weights.GetBuffer(), PF_R32_FLOAT));
 			if (HasBias) {
@@ -152,16 +155,16 @@ namespace UE::NNIRuntimeRDG::Private::Hlsl
 			PermutationVector.Set<FConvCS::FConvAlgorithm>(Algorithm);
 			PermutationVector.Set<FConvCS::FConvGroupSize>(GroupSize);
 			PermutationVector.Set<FConvCS::FConvNumDimensions>(NumDimensions);
-			PermutationVector.Set<FConvCS::FConvNumReadsPerThread>(FConvCS::GetNumReadsPerThread(GroupSize, Weights.GetShape().Data, Dilations, Strides));
+			PermutationVector.Set<FConvCS::FConvNumReadsPerThread>(FConvCS::GetNumReadsPerThread(GroupSize, Weights.GetShape().GetData(), Dilations, Strides));
 			PermutationVector.Set<FConvCS::FConvHasB>(HasBias);
 			TShaderMapRef<FConvCS> ComputeShader(GetGlobalShaderMap(GMaxRHIFeatureLevel), PermutationVector);
 
-			RDG_EVENT_SCOPE(GraphBuilder, "NNI.Operator.Hlsl.Conv");
-			RDG_GPU_STAT_SCOPE(GraphBuilder, FNNIOperatorConv);
+			RDG_EVENT_SCOPE(GraphBuilder, "NNE.Operator.Hlsl.Conv");
+			RDG_GPU_STAT_SCOPE(GraphBuilder, FNNEOperatorConv);
 
 			FComputeShaderUtils::AddPass(
 				GraphBuilder,
-				RDG_EVENT_NAME("NNI.Operator.Hlsl.Conv.Dispatch"),
+				RDG_EVENT_NAME("NNE.Operator.Hlsl.Conv.Dispatch"),
 				ERDGPassFlags::Compute | ERDGPassFlags::NeverCull,
 				ComputeShader,
 				Params,
@@ -169,7 +172,7 @@ namespace UE::NNIRuntimeRDG::Private::Hlsl
 		}
 	};
 
-	bool ValidateConvOperator(const UE::NNECore::FAttributeMap& AttributeMap, TConstArrayView<EMLTensorDataType> InputTypes, TConstArrayView<NNX::FSymbolicTensorShape> InputShapes)
+	bool ValidateConvOperator(const NNECore::FAttributeMap& AttributeMap, TConstArrayView<ENNETensorDataType> InputTypes, TConstArrayView<NNECore::FSymbolicTensorShape> InputShapes)
 	{
 		bool bIsValid = true;
 
@@ -184,7 +187,7 @@ namespace UE::NNIRuntimeRDG::Private::Hlsl
 		bIsValid &= AttributeValidator.Validate(AttributeMap);
 
 		NNX::FInputValidator InputValidator;
-		InputValidator.AddSupportedType(EMLTensorDataType::Float);
+		InputValidator.AddSupportedType(ENNETensorDataType::Float);
 		InputValidator.AddRequired();
 		InputValidator.AddRequired();
 		InputValidator.AddOptional();
diff --git a/Engine/Plugins/Experimental/NNE/Source/NNXRuntimeRDG/Private/Hlsl/NNIRuntimeRDGConvTranspose.cpp b/Engine/Plugins/Experimental/NNE/Source/NNXRuntimeRDG/Private/Hlsl/NNIRuntimeRDGConvTranspose.cpp
--- a/Engine/Plugins/Experimental/NNE/Source/NNXRuntimeRDG/Private/Hlsl/NNIRuntimeRDGConvTranspose.cpp
+++ b/Engine/Plugins/Experimental/NNE/Source/NNXRuntimeRDG/Private/Hlsl/NNIRuntimeRDGConvTranspose.cpp
@@ -3,10 +3,12 @@
 #include "NNIRuntimeRDGConvTranspose.h"
 #include "NNEHlslShadersConvTransposeCS.h"
 #include "NNXRuntimeHLSLHelper.h"
+#include "NNECoreTensor.h"
+#include "NNECoreTypes.h"
 
 namespace UE::NNIRuntimeRDG::Private::Hlsl
 {
-	DECLARE_GPU_STAT_NAMED(FNNIOperatorConvTranspose, TEXT("NNI.Operator.Hlsl.ConvTranspose"));
+	DECLARE_GPU_STAT_NAMED(FNNEOperatorConvTranspose, TEXT("NNE.Operator.Hlsl.ConvTranspose"));
 
 	using EConvTransposeAutoPad = UE::NNEHlslShaders::Internal::EConvTransposeAutoPad;
 
@@ -39,35 +41,36 @@ namespace UE::NNIRuntimeRDG::Private::Hlsl
 
 	public:
 
-		virtual int PrepareOutputs(TConstArrayView<NNX::FTensorRef> InputTensors, TArrayView<NNX::FTensorRef> OutputTensors) const override
+		virtual int PrepareOutputs(TConstArrayView<NNECore::Internal::FTensorRef> InputTensors, TArrayView<NNECore::Internal::FTensorRef> OutputTensors) const override
 		{
 			check(InputTensors.Num() >= 2 && InputTensors.Num() <= 3);
 			check(OutputTensors.Num() == 1);
 			
-			const NNX::FTensorShape& Input = InputTensors[0]->GetShape();
-			const NNX::FTensorShape& Weights = InputTensors[1]->GetShape();
-			NNX::FSymbolicTensorShape OutputShape;
+			const NNECore::FTensorShape& Input = InputTensors[0]->GetShape();
+			const NNECore::FTensorShape& Weights = InputTensors[1]->GetShape();
+
 			
-			OutputShape.Data = UE::NNEHlslShaders::Internal::FConvTransposeCS::GetOutputShape(Input.Data, Weights.Data, AutoPad, Dilations, Strides, Pads, OutputPadding, Group);
+			const TArray<int32> OutputShapeData = NNEHlslShaders::Internal::FConvTransposeCS::GetOutputShape(Input.GetData(), Weights.GetData(), AutoPad, Dilations, Strides, Pads, OutputPadding, Group);
+			const NNECore::FSymbolicTensorShape OutputShape = NNECore::FSymbolicTensorShape::Make(OutputShapeData);
 			if (!OutputShape.IsConcrete())
 			{
 				return -1;
 			}
-			OutputTensors[0]->SetShape(NNX::FTensorShape::MakeFromSymbolic(OutputShape));
+			OutputTensors[0]->SetShape(NNECore::FTensorShape::MakeFromSymbolic(OutputShape));
 
 			return 0;
 		};
 
-		virtual bool Initialize(TConstArrayView<NNX::FTensorDesc> InputTensorDescs, TConstArrayView<NNX::FTensorDesc> OutputTensorDescs, const UE::NNECore::FAttributeMap& Attributes) override
+		virtual bool Initialize(TConstArrayView<NNECore::FTensorDesc> InputTensorDescs, TConstArrayView<NNECore::FTensorDesc> OutputTensorDescs, const NNECore::FAttributeMap& Attributes) override
 		{
 			using namespace UE::NNEHlslShaders::Internal;
 
 			check(InputTensorDescs.Num() >= 2 && InputTensorDescs.Num() <= 3);
 			check(OutputTensorDescs.Num() == 1);
 
-            const NNX::FTensorDesc& Input = InputTensorDescs[0];
-			const NNX::FTensorDesc& Weights = InputTensorDescs[1];
-			const NNX::FTensorDesc& Output = OutputTensorDescs[0];
+            const NNECore::FTensorDesc& Input = InputTensorDescs[0];
+			const NNECore::FTensorDesc& Weights = InputTensorDescs[1];
+			const NNECore::FTensorDesc& Output = OutputTensorDescs[0];
 
 			if (Input.GetShape().Rank() < 2)
 			{
@@ -140,11 +143,11 @@ namespace UE::NNIRuntimeRDG::Private::Hlsl
 			check(Output.GetShape().Rank() == Input.GetShape().Rank());
 			check(NumDimensions == (Input.GetShape().Rank() - 2));
 
-			TArray<int32> OutputShape = FConvTransposeCS::GetOutputShape(Input.GetShape().Data, Weights.GetShape().Data, AutoPad, Dilations, Strides, Pads, OutputPadding, Group);
+			TArray<int32> OutputShape = FConvTransposeCS::GetOutputShape(Input.GetShape().GetData(), Weights.GetShape().GetData(), AutoPad, Dilations, Strides, Pads, OutputPadding, Group);
 
 			// Set parameters
 			FConvTransposeCS::FParameters* Params = GraphBuilder.AllocParameters<FConvTransposeCS::FParameters>();
-			FConvTransposeCS::FillInParameters(GroupSize, Input.GetShape().Data, Weights.GetShape().Data, HasBias, AutoPad, Group, Dilations,Strides, Pads, OutputPadding, *Params);
+			FConvTransposeCS::FillInParameters(GroupSize, Input.GetShape().GetData(), Weights.GetShape().GetData(), HasBias, AutoPad, Group, Dilations,Strides, Pads, OutputPadding, *Params);
 			Params->X = GraphBuilder.CreateSRV(FRDGBufferSRVDesc(Input.GetBuffer(), PF_R32_FLOAT));
 			Params->W = GraphBuilder.CreateSRV(FRDGBufferSRVDesc(Weights.GetBuffer(), PF_R32_FLOAT));
 			if (HasBias) {
@@ -157,16 +160,16 @@ namespace UE::NNIRuntimeRDG::Private::Hlsl
 			PermutationVector.Set<FConvTransposeCS::FConvTransposeAlgorithm>(Algorithm);
 			PermutationVector.Set<FConvTransposeCS::FConvTransposeGroupSize>(GroupSize);
 			PermutationVector.Set<FConvTransposeCS::FConvTransposeNumStackDimensions>(NumDimensions);
-			PermutationVector.Set<FConvTransposeCS::FConvTransposeNumReadsPerThread>(FConvTransposeCS::GetNumReadsPerThread(GroupSize, Weights.GetShape().Data, Dilations, Strides));
+			PermutationVector.Set<FConvTransposeCS::FConvTransposeNumReadsPerThread>(FConvTransposeCS::GetNumReadsPerThread(GroupSize, Weights.GetShape().GetData(), Dilations, Strides));
 			PermutationVector.Set<FConvTransposeCS::FConvTransposeHasB>(HasBias);
 			TShaderMapRef<FConvTransposeCS> ComputeShader(GetGlobalShaderMap(GMaxRHIFeatureLevel), PermutationVector);
 
-			RDG_EVENT_SCOPE(GraphBuilder, "NNI.Operator.Hlsl.ConvTranspose");
-			RDG_GPU_STAT_SCOPE(GraphBuilder, FNNIOperatorConvTranspose);
+			RDG_EVENT_SCOPE(GraphBuilder, "NNE.Operator.Hlsl.ConvTranspose");
+			RDG_GPU_STAT_SCOPE(GraphBuilder, FNNEOperatorConvTranspose);
 
 			FComputeShaderUtils::AddPass(
 				GraphBuilder,
-				RDG_EVENT_NAME("NNI.Operator.Hlsl.ConvTranspose.Dispatch"),
+				RDG_EVENT_NAME("NNE.Operator.Hlsl.ConvTranspose.Dispatch"),
 				ERDGPassFlags::Compute | ERDGPassFlags::NeverCull,
 				ComputeShader,
 				Params,
@@ -174,7 +177,7 @@ namespace UE::NNIRuntimeRDG::Private::Hlsl
 		}
 	};
 
-	bool ValidateConvTransposeOperator(const UE::NNECore::FAttributeMap& AttributeMap, TConstArrayView<EMLTensorDataType> InputTypes, TConstArrayView<NNX::FSymbolicTensorShape> InputShapes)
+	bool ValidateConvTransposeOperator(const NNECore::FAttributeMap& AttributeMap, TConstArrayView<ENNETensorDataType> InputTypes, TConstArrayView<NNECore::FSymbolicTensorShape> InputShapes)
 	{
 		bool bIsValid = true;
 
@@ -191,7 +194,7 @@ namespace UE::NNIRuntimeRDG::Private::Hlsl
 		bIsValid &= AttributeValidator.Validate(AttributeMap);
 
 		NNX::FInputValidator InputValidator;
-		InputValidator.AddSupportedType(EMLTensorDataType::Float);
+		InputValidator.AddSupportedType(ENNETensorDataType::Float);
 		InputValidator.AddRequired();
 		InputValidator.AddRequired();
 		InputValidator.AddOptional();
diff --git a/Engine/Plugins/Experimental/NNE/Source/NNXRuntimeRDG/Private/Hlsl/NNIRuntimeRDGElementWiseBinary.cpp b/Engine/Plugins/Experimental/NNE/Source/NNXRuntimeRDG/Private/Hlsl/NNIRuntimeRDGElementWiseBinary.cpp
--- a/Engine/Plugins/Experimental/NNE/Source/NNXRuntimeRDG/Private/Hlsl/NNIRuntimeRDGElementWiseBinary.cpp
+++ b/Engine/Plugins/Experimental/NNE/Source/NNXRuntimeRDG/Private/Hlsl/NNIRuntimeRDGElementWiseBinary.cpp
@@ -3,10 +3,12 @@
 #include "NNIRuntimeRDGElementWiseBinary.h"
 #include "NNEHlslShadersElementWiseBinaryCS.h"
 #include "NNXRuntimeHLSLHelper.h"
+#include "NNECoreTensor.h"
+#include "NNECoreTypes.h"
 
 namespace UE::NNIRuntimeRDG::Private::Hlsl
 {
-	DECLARE_GPU_STAT_NAMED(FNNIOperatorElementWiseBinary, TEXT("NNI.Operator.Hlsl.ElementWise.Binary"));
+	DECLARE_GPU_STAT_NAMED(FNNEOperatorElementWiseBinary, TEXT("NNE.Operator.Hlsl.ElementWise.Binary"));
 
 	using TElementWiseBinaryCS = typename UE::NNEHlslShaders::Internal::TElementWiseBinaryCS;
 	using FElementWiseBinaryConstants = UE::NNEHlslShaders::Internal::FElementWiseBinaryConstants;
@@ -24,39 +26,41 @@ namespace UE::NNIRuntimeRDG::Private::Hlsl
 
 	public:
 
-		virtual int PrepareOutputs(TConstArrayView<NNX::FTensorRef> InputTensors, TArrayView<NNX::FTensorRef> OutputTensors) const override
+		virtual int PrepareOutputs(TConstArrayView<NNECore::Internal::FTensorRef> InputTensors, TArrayView<NNECore::Internal::FTensorRef> OutputTensors) const override
 		{
 			check(InputTensors.Num() == 2);
 			check(OutputTensors.Num() == 1);
 			
-			const NNX::FTensorShape& LHSInput = InputTensors[0]->GetShape();
-			const NNX::FTensorShape& RHSInput = InputTensors[1]->GetShape();
+			const NNECore::FTensorShape& LHSInput = InputTensors[0]->GetShape();
+			const NNECore::FTensorShape& RHSInput = InputTensors[1]->GetShape();
 			const int32 OutputRank = FMath::Max(LHSInput.Rank(), RHSInput.Rank());
-			NNX::FTensorShape OutputShape;
+			TArray<uint32> OutputShapeData;
 			
-			OutputShape.Data.SetNumUninitialized(OutputRank);
+			OutputShapeData.SetNumUninitialized(OutputRank);
 			
 			for (int32 i = 0; i < OutputRank; ++i)
 			{
 				int32 LHSIndex = LHSInput.Rank() - 1 - i;
 				int32 RHSIndex = RHSInput.Rank() - 1 - i;
-				int32 LHSValue = LHSIndex >= 0 ? LHSInput.Data[LHSIndex] : 1;
-				int32 RHSValue = RHSIndex >= 0 ? RHSInput.Data[RHSIndex] : 1;
+				int32 LHSValue = LHSIndex >= 0 ? LHSInput.GetData()[LHSIndex] : 1;
+				int32 RHSValue = RHSIndex >= 0 ? RHSInput.GetData()[RHSIndex] : 1;
 				if (LHSValue != RHSValue && LHSValue != 1 && RHSValue != 1)
 				{
 					UE_LOG(LogNNX, Warning, TEXT("Error while computing shape for element wise binary op, input shapes are not compatible"));
 					return -1;
 				}
 				int32 OutputValue = FMath::Max(LHSValue, RHSValue);
-				OutputShape.Data[OutputRank - 1 - i] = OutputValue;
+				OutputShapeData[OutputRank - 1 - i] = OutputValue;
 			}
 
+			NNECore::FTensorShape OutputShape = NNECore::FTensorShape::Make(OutputShapeData);
+			
 			OutputTensors[0]->SetShape(OutputShape);
 			
 			return 0;
 		}
 		
-		virtual bool Initialize(TConstArrayView<NNX::FTensorDesc> InputTensorDescs, TConstArrayView<NNX::FTensorDesc> OutputTensorDescs, const UE::NNECore::FAttributeMap& Attributes) override
+		virtual bool Initialize(TConstArrayView<NNECore::FTensorDesc> InputTensorDescs, TConstArrayView<NNX::FTensorDesc> OutputTensorDescs, const NNECore::FAttributeMap& Attributes) override
 		{
 			check(InputTensorDescs.Num() == 2);
 			check(OutputTensorDescs.Num() == 1);
@@ -99,12 +103,12 @@ namespace UE::NNIRuntimeRDG::Private::Hlsl
 
 			TShaderMapRef<TElementWiseBinaryCS> ComputeShader(GetGlobalShaderMap(GMaxRHIFeatureLevel), PermutationVector);
 		
-			RDG_EVENT_SCOPE(GraphBuilder, "NNI.Operator.Hlsl.ElementWise.Binary");
-			RDG_GPU_STAT_SCOPE(GraphBuilder, FNNIOperatorElementWiseBinary);
+			RDG_EVENT_SCOPE(GraphBuilder, "NNE.Operator.Hlsl.ElementWise.Binary");
+			RDG_GPU_STAT_SCOPE(GraphBuilder, FNNEOperatorElementWiseBinary);
 
 			FComputeShaderUtils::AddPass(
 				GraphBuilder,
-				RDG_EVENT_NAME("NNI.Operator.Hlsl.ElementWise.Binary.Dispatch"),
+				RDG_EVENT_NAME("NNE.Operator.Hlsl.ElementWise.Binary.Dispatch"),
 				ERDGPassFlags::Compute | ERDGPassFlags::NeverCull,
 				ComputeShader,
 				Params,
@@ -112,7 +116,7 @@ namespace UE::NNIRuntimeRDG::Private::Hlsl
 		}
 	};
 
-	bool ValidateElementWiseBinaryOperator(const UE::NNECore::FAttributeMap& AttributeMap, TConstArrayView<EMLTensorDataType> InputTypes, TConstArrayView<NNX::FSymbolicTensorShape> InputShapes)
+	bool ValidateElementWiseBinaryOperator(const NNECore::FAttributeMap& AttributeMap, TConstArrayView<ENNETensorDataType> InputTypes, TConstArrayView<NNECore::FSymbolicTensorShape> InputShapes)
 	{
 		bool bIsValid = true;
 
@@ -120,7 +124,7 @@ namespace UE::NNIRuntimeRDG::Private::Hlsl
 		bIsValid &= AttributeValidator.Validate(AttributeMap);
 
 		NNX::FInputValidator InputValidator;
-		InputValidator.AddSupportedType(EMLTensorDataType::Float);
+		InputValidator.AddSupportedType(ENNETensorDataType::Float);
 		InputValidator.AddRequired();
 		InputValidator.AddRequired();
 		bIsValid &= InputValidator.Validate(InputTypes);
diff --git a/Engine/Plugins/Experimental/NNE/Source/NNXRuntimeRDG/Private/Hlsl/NNIRuntimeRDGElementWiseUnary.cpp b/Engine/Plugins/Experimental/NNE/Source/NNXRuntimeRDG/Private/Hlsl/NNIRuntimeRDGElementWiseUnary.cpp
--- a/Engine/Plugins/Experimental/NNE/Source/NNXRuntimeRDG/Private/Hlsl/NNIRuntimeRDGElementWiseUnary.cpp
+++ b/Engine/Plugins/Experimental/NNE/Source/NNXRuntimeRDG/Private/Hlsl/NNIRuntimeRDGElementWiseUnary.cpp
@@ -4,10 +4,12 @@
 #include "NNEHlslShadersElementWiseUnaryCS.h"
 #include "NNXRuntimeHLSLHelper.h"
 #include "NNECoreAttributeMap.h"
+#include "NNECoreTypes.h"
+#include "NNECoreTensor.h"
 
 namespace UE::NNIRuntimeRDG::Private::Hlsl
 {
-	DECLARE_GPU_STAT_NAMED(FNNIOperatorElementWiseUnary, TEXT("NNI.Operator.Hlsl.ElementWise.Unary"));
+	DECLARE_GPU_STAT_NAMED(FNNEOperatorElementWiseUnary, TEXT("NNE.Operator.Hlsl.ElementWise.Unary"));
 
 	using TElementWiseUnaryCS = typename UE::NNEHlslShaders::Internal::TElementWiseUnaryCS;
 	using FElementWiseUnaryConstants = UE::NNEHlslShaders::Internal::FElementWiseUnaryConstants;
@@ -31,7 +33,7 @@ namespace UE::NNIRuntimeRDG::Private::Hlsl
 
 	public:
 
-		virtual int PrepareOutputs(TConstArrayView<NNX::FTensorRef> InputTensors, TArrayView<NNX::FTensorRef> OutputTensors) const override
+		virtual int PrepareOutputs(TConstArrayView<NNECore::Internal::FTensorRef> InputTensors, TArrayView<NNECore::Internal::FTensorRef> OutputTensors) const override
 		{
 			check(InputTensors.Num() == 1);
 			check(OutputTensors.Num() == 1);
@@ -39,7 +41,7 @@ namespace UE::NNIRuntimeRDG::Private::Hlsl
 			return 0;
 		}
 
-		virtual bool Initialize(TConstArrayView<NNX::FTensorDesc> InputTensorDescs, TConstArrayView<NNX::FTensorDesc> OutputTensorDescs, const UE::NNECore::FAttributeMap& Attributes) override
+		virtual bool Initialize(TConstArrayView<NNECore::FTensorDesc> InputTensorDescs, TConstArrayView<NNECore::FTensorDesc> OutputTensorDescs, const NNECore::FAttributeMap& Attributes) override
 		{
 			check(InputTensorDescs.Num() == 1);
 			check(OutputTensorDescs.Num() == 1);
@@ -76,8 +78,8 @@ namespace UE::NNIRuntimeRDG::Private::Hlsl
 
 			TShaderMapRef<TElementWiseUnaryCS> ComputeShader(GetGlobalShaderMap(GMaxRHIFeatureLevel), PermutationVector);
 
-			RDG_EVENT_SCOPE(GraphBuilder, "NNI.Operator.Hlsl.ElementWise.Unary");
-			RDG_GPU_STAT_SCOPE(GraphBuilder, FNNIOperatorElementWiseUnary);
+			RDG_EVENT_SCOPE(GraphBuilder, "NNE.Operator.Hlsl.ElementWise.Unary");
+			RDG_GPU_STAT_SCOPE(GraphBuilder, FNNEOperatorElementWiseUnary);
 		
 			FComputeShaderUtils::AddPass(
 				GraphBuilder,
@@ -116,7 +118,7 @@ namespace UE::NNIRuntimeRDG::Private::Hlsl
 	}
 
 	template<EMLElementWiseUnaryOperatorType OpType>
-	bool ValidateElementWiseUnaryOperator(const UE::NNECore::FAttributeMap& AttributeMap, TConstArrayView<EMLTensorDataType> InputTypes, TConstArrayView<NNX::FSymbolicTensorShape> InputShapes)
+	bool ValidateElementWiseUnaryOperator(const NNECore::FAttributeMap& AttributeMap, TConstArrayView<ENNETensorDataType> InputTypes, TConstArrayView<NNECore::FSymbolicTensorShape> InputShapes)
 	{
 		bool bIsValid = true;
 
@@ -124,7 +126,7 @@ namespace UE::NNIRuntimeRDG::Private::Hlsl
 		bIsValid &= AttributeValidator.Validate(AttributeMap);
 
 		NNX::FInputValidator InputValidator;
-		InputValidator.AddSupportedType(EMLTensorDataType::Float);
+		InputValidator.AddSupportedType(ENNETensorDataType::Float);
 		InputValidator.AddRequired();
 		bIsValid &= InputValidator.Validate(InputTypes);
 
@@ -132,7 +134,7 @@ namespace UE::NNIRuntimeRDG::Private::Hlsl
 	}
 
 	template<>
-	bool ValidateElementWiseUnaryOperator<EMLElementWiseUnaryOperatorType::Selu>(const UE::NNECore::FAttributeMap& AttributeMap, TConstArrayView<EMLTensorDataType> InputTypes, TConstArrayView<NNX::FSymbolicTensorShape> InputShapes)
+	bool ValidateElementWiseUnaryOperator<EMLElementWiseUnaryOperatorType::Selu>(const NNECore::FAttributeMap& AttributeMap, TConstArrayView<ENNETensorDataType> InputTypes, TConstArrayView<NNECore::FSymbolicTensorShape> InputShapes)
 	{
 		bool bIsValid = true;
 
@@ -142,7 +144,7 @@ namespace UE::NNIRuntimeRDG::Private::Hlsl
 		bIsValid &= AttributeValidator.Validate(AttributeMap);
 
 		NNX::FInputValidator InputValidator;
-		InputValidator.AddSupportedType(EMLTensorDataType::Float);
+		InputValidator.AddSupportedType(ENNETensorDataType::Float);
 		InputValidator.AddRequired();
 		bIsValid &= InputValidator.Validate(InputTypes);
 
@@ -150,7 +152,7 @@ namespace UE::NNIRuntimeRDG::Private::Hlsl
 	}
 
 	template<>
-	bool ValidateElementWiseUnaryOperator<EMLElementWiseUnaryOperatorType::Elu>(const UE::NNECore::FAttributeMap& AttributeMap, TConstArrayView<EMLTensorDataType> InputTypes, TConstArrayView<NNX::FSymbolicTensorShape> InputShapes)
+	bool ValidateElementWiseUnaryOperator<EMLElementWiseUnaryOperatorType::Elu>(const NNECore::FAttributeMap& AttributeMap, TConstArrayView<ENNETensorDataType> InputTypes, TConstArrayView<NNECore::FSymbolicTensorShape> InputShapes)
 	{
 		bool bIsValid = true;
 
@@ -159,7 +161,7 @@ namespace UE::NNIRuntimeRDG::Private::Hlsl
 		bIsValid &= AttributeValidator.Validate(AttributeMap);
 
 		NNX::FInputValidator InputValidator;
-		InputValidator.AddSupportedType(EMLTensorDataType::Float);
+		InputValidator.AddSupportedType(ENNETensorDataType::Float);
 		InputValidator.AddRequired();
 		bIsValid &= InputValidator.Validate(InputTypes);
 
@@ -167,7 +169,7 @@ namespace UE::NNIRuntimeRDG::Private::Hlsl
 	}
 
 	template<>
-	bool ValidateElementWiseUnaryOperator<EMLElementWiseUnaryOperatorType::HardSigmoid>(const UE::NNECore::FAttributeMap& AttributeMap, TConstArrayView<EMLTensorDataType> InputTypes, TConstArrayView<NNX::FSymbolicTensorShape> InputShapes)
+	bool ValidateElementWiseUnaryOperator<EMLElementWiseUnaryOperatorType::HardSigmoid>(const NNECore::FAttributeMap& AttributeMap, TConstArrayView<ENNETensorDataType> InputTypes, TConstArrayView<NNECore::FSymbolicTensorShape> InputShapes)
 	{
 		bool bIsValid = true;
 
@@ -177,7 +179,7 @@ namespace UE::NNIRuntimeRDG::Private::Hlsl
 		bIsValid &= AttributeValidator.Validate(AttributeMap);
 
 		NNX::FInputValidator InputValidator;
-		InputValidator.AddSupportedType(EMLTensorDataType::Float);
+		InputValidator.AddSupportedType(ENNETensorDataType::Float);
 		InputValidator.AddRequired();
 		bIsValid &= InputValidator.Validate(InputTypes);
 
@@ -185,7 +187,7 @@ namespace UE::NNIRuntimeRDG::Private::Hlsl
 	}
 
 	template<>
-	bool ValidateElementWiseUnaryOperator<EMLElementWiseUnaryOperatorType::LeakyRelu>(const UE::NNECore::FAttributeMap& AttributeMap, TConstArrayView<EMLTensorDataType> InputTypes, TConstArrayView<NNX::FSymbolicTensorShape> InputShapes)
+	bool ValidateElementWiseUnaryOperator<EMLElementWiseUnaryOperatorType::LeakyRelu>(const NNECore::FAttributeMap& AttributeMap, TConstArrayView<ENNETensorDataType> InputTypes, TConstArrayView<NNECore::FSymbolicTensorShape> InputShapes)
 	{
 		bool bIsValid = true;
 
@@ -194,7 +196,7 @@ namespace UE::NNIRuntimeRDG::Private::Hlsl
 		bIsValid &= AttributeValidator.Validate(AttributeMap);
 
 		NNX::FInputValidator InputValidator;
-		InputValidator.AddSupportedType(EMLTensorDataType::Float);
+		InputValidator.AddSupportedType(ENNETensorDataType::Float);
 		InputValidator.AddRequired();
 		bIsValid &= InputValidator.Validate(InputTypes);
 
diff --git a/Engine/Plugins/Experimental/NNE/Source/NNXRuntimeRDG/Private/Hlsl/NNIRuntimeRDGElementWiseVariadic.cpp b/Engine/Plugins/Experimental/NNE/Source/NNXRuntimeRDG/Private/Hlsl/NNIRuntimeRDGElementWiseVariadic.cpp
--- a/Engine/Plugins/Experimental/NNE/Source/NNXRuntimeRDG/Private/Hlsl/NNIRuntimeRDGElementWiseVariadic.cpp
+++ b/Engine/Plugins/Experimental/NNE/Source/NNXRuntimeRDG/Private/Hlsl/NNIRuntimeRDGElementWiseVariadic.cpp
@@ -3,10 +3,12 @@
 #include "NNIRuntimeRDGElementWiseVariadic.h"
 #include "NNEHlslShadersElementWiseVariadicCS.h"
 #include "NNXRuntimeHLSLHelper.h"
+#include "NNECoreTensor.h"
+#include "NNECoreTypes.h"
 
 namespace UE::NNIRuntimeRDG::Private::Hlsl
 {
-	DECLARE_GPU_STAT_NAMED(FNNIOperatorElementWiseVariadic, TEXT("NNI.Operator.Hlsl.ElementWise.Variadic"));
+	DECLARE_GPU_STAT_NAMED(FNNEOperatorElementWiseVariadic, TEXT("NNE.Operator.Hlsl.ElementWise.Variadic"));
 
 	using TElementWiseVariadicCS = typename UE::NNEHlslShaders::Internal::TElementWiseVariadicCS;
 	using FElementWiseVariadicConstants = UE::NNEHlslShaders::Internal::FElementWiseVariadicConstants;
@@ -92,7 +94,7 @@ namespace UE::NNIRuntimeRDG::Private::Hlsl
 
 	public:
 
-		virtual int PrepareOutputs(TConstArrayView<NNX::FTensorRef> InputTensors, TArrayView<NNX::FTensorRef> OutputTensors) const override
+		virtual int PrepareOutputs(TConstArrayView<NNECore::Internal::FTensorRef> InputTensors, TArrayView<NNECore::Internal::FTensorRef> OutputTensors) const override
 		{
 			check(InputTensors.Num() > 0);
 			check(OutputTensors.Num() == 1);
@@ -104,9 +106,9 @@ namespace UE::NNIRuntimeRDG::Private::Hlsl
 				OutputRank = FMath::Max(OutputRank, InputTensors[i]->GetShape().Rank());
 			}
 
-			NNX::FTensorShape OutputShape;
+			TArray<uint32> OutputShapeData;
 
-			OutputShape.Data.SetNumUninitialized(OutputRank);
+			OutputShapeData.SetNumUninitialized(OutputRank);
 
 			for (int32 i = 0; i < OutputRank; ++i)
 			{
@@ -114,7 +116,7 @@ namespace UE::NNIRuntimeRDG::Private::Hlsl
 				for (int32 InputIdx = 0; InputIdx < NumInput; ++InputIdx)
 				{
 					int32 InputIndex = InputTensors[InputIdx]->GetShape().Rank() - 1 - i;
-					int32 InputValue = InputIndex >= 0 ? InputTensors[InputIdx]->GetShape().Data[InputIndex] : 1;
+					int32 InputValue = InputIndex >= 0 ? InputTensors[InputIdx]->GetShape().GetData()[InputIndex] : 1;
 					if (InputValue != OutputValue && InputValue != 1 && OutputValue != 1)
 					{
 						UE_LOG(LogNNX, Warning, TEXT("Error while computing shape for element wise variadic op, input shapes are not compatible"));
@@ -122,9 +124,11 @@ namespace UE::NNIRuntimeRDG::Private::Hlsl
 					}
 					OutputValue = FMath::Max(InputValue, OutputValue);
 				}
-				OutputShape.Data[OutputRank - 1 - i] = OutputValue;
+				OutputShapeData[OutputRank - 1 - i] = OutputValue;
 			}
 
+			NNECore::FTensorShape OutputShape = NNECore::FTensorShape::Make(OutputShapeData);
+
 			OutputTensors[0]->SetShape(OutputShape);
 
 			return 0;
@@ -143,8 +147,8 @@ namespace UE::NNIRuntimeRDG::Private::Hlsl
 			check(InOutputTensors.Num() == 1);
 			check(InOutputTensors[0] != nullptr);
 
-			RDG_EVENT_SCOPE(GraphBuilder, "NNI.Operator.Hlsl.ElementWise.Variadic");
-			RDG_GPU_STAT_SCOPE(GraphBuilder, FNNIOperatorElementWiseVariadic);
+			RDG_EVENT_SCOPE(GraphBuilder, "NNE.Operator.Hlsl.ElementWise.Variadic");
+			RDG_GPU_STAT_SCOPE(GraphBuilder, FNNEOperatorElementWiseVariadic);
 
 			NNX::FTensorRDGRef PassInputTensors[FElementWiseVariadicConstants::MAX_NUM_INPUT];
 			for (int32 InputOffset = 0; InputOffset < InInputTensors.Num(); InputOffset += FElementWiseVariadicConstants::MAX_NUM_INPUT)
@@ -175,7 +179,7 @@ namespace UE::NNIRuntimeRDG::Private::Hlsl
 		}
 	};
 
-	bool ValidateElementWiseVariadicOperator(const UE::NNECore::FAttributeMap& AttributeMap, TConstArrayView<EMLTensorDataType> InputTypes, TConstArrayView<NNX::FSymbolicTensorShape> InputShapes)
+	bool ValidateElementWiseVariadicOperator(const NNECore::FAttributeMap& AttributeMap, TConstArrayView<ENNETensorDataType> InputTypes, TConstArrayView<NNECore::FSymbolicTensorShape> InputShapes)
 	{
 		bool bIsValid = true;
 
@@ -189,7 +193,7 @@ namespace UE::NNIRuntimeRDG::Private::Hlsl
 		}
 		for (int32 i = 0; i < InputTypes.Num(); ++i)
 		{
-			if (InputTypes[i] != EMLTensorDataType::Float)
+			if (InputTypes[i] != ENNETensorDataType::Float)
 			{
 				UE_LOG(LogNNX, Warning, TEXT("Element-wise variadic operator input '%d' of type '%d' is not supported, should be float at the moment."), i, InputTypes[i]);
 				bIsValid = false;
@@ -197,7 +201,7 @@ namespace UE::NNIRuntimeRDG::Private::Hlsl
 		}
 		
 		NNX::FInputValidator InputValidator;
-		InputValidator.AddSupportedType(EMLTensorDataType::Float);
+		InputValidator.AddSupportedType(ENNETensorDataType::Float);
 		InputValidator.AddRequired();
 		InputValidator.AddRequired();
 		bIsValid &= InputValidator.Validate(InputTypes);
diff --git a/Engine/Plugins/Experimental/NNE/Source/NNXRuntimeRDG/Private/Hlsl/NNIRuntimeRDGGather.cpp b/Engine/Plugins/Experimental/NNE/Source/NNXRuntimeRDG/Private/Hlsl/NNIRuntimeRDGGather.cpp
--- a/Engine/Plugins/Experimental/NNE/Source/NNXRuntimeRDG/Private/Hlsl/NNIRuntimeRDGGather.cpp
+++ b/Engine/Plugins/Experimental/NNE/Source/NNXRuntimeRDG/Private/Hlsl/NNIRuntimeRDGGather.cpp
@@ -4,10 +4,12 @@
 #include "NNXRuntimeHLSLHelper.h"
 #include "NNEHlslShadersGatherCS.h"
 #include "NNECoreAttributeMap.h"
+#include "NNECoreTypes.h"
+#include "NNECoreTensor.h"
 
 namespace UE::NNIRuntimeRDG::Private::Hlsl
 {
-	DECLARE_GPU_STAT_NAMED(FNNIOperatorGather, TEXT("NNI.Operator.Hlsl.Gather"));
+	DECLARE_GPU_STAT_NAMED(FNNEOperatorGather, TEXT("NNE.Operator.Hlsl.Gather"));
 
 	/**
 	 * Gather operator implementation
@@ -26,7 +28,7 @@ namespace UE::NNIRuntimeRDG::Private::Hlsl
 
 	public:
 
-		virtual int PrepareOutputs(TConstArrayView<NNX::FTensorRef> InputTensors, TArrayView<NNX::FTensorRef> OutputTensors) const override
+		virtual int PrepareOutputs(TConstArrayView<NNECore::Internal::FTensorRef> InputTensors, TArrayView<NNECore::Internal::FTensorRef> OutputTensors) const override
 		{
 			check(InputTensors.Num() == 2)
 			check(OutputTensors.Num() == 1)
@@ -34,16 +36,16 @@ namespace UE::NNIRuntimeRDG::Private::Hlsl
 			return -1;
 		};
 		
-		virtual bool Initialize(TConstArrayView<NNX::FTensorDesc> InputTensorDescs, TConstArrayView<NNX::FTensorDesc> OutputTensorDescs, const UE::NNECore::FAttributeMap& Attributes) override
+		virtual bool Initialize(TConstArrayView<NNECore::FTensorDesc> InputTensorDescs, TConstArrayView<NNECore::FTensorDesc> OutputTensorDescs, const UE::NNECore::FAttributeMap& Attributes) override
 		{
 			const int32 MaxNumDimensions = UE::NNEHlslShaders::Internal::FGatherConstants::MAX_NUM_DIMENSIONS;
 			
 			check(InputTensorDescs.Num() == 2)
 			check(OutputTensorDescs.Num() == 1)
 
-			const NNX::FTensorDesc& Data = InputTensorDescs[0];
-			const NNX::FTensorDesc& Indices = InputTensorDescs[1];
-			const NNX::FTensorDesc& Output = OutputTensorDescs[0];
+			const NNECore::FTensorDesc& Data = InputTensorDescs[0];
+			const NNECore::FTensorDesc& Indices = InputTensorDescs[1];
+			const NNECore::FTensorDesc& Output = OutputTensorDescs[0];
 
 			if (Output.GetShape().Rank() <= MaxNumDimensions)
 			{
@@ -113,8 +115,8 @@ namespace UE::NNIRuntimeRDG::Private::Hlsl
 
 			FIntVector ThreadGroupCount = TGatherCS::GetGroupCount(*Parameters);
 
-			RDG_EVENT_SCOPE(GraphBuilder, "NNI.Operator.Hlsl.Gather");
-			RDG_GPU_STAT_SCOPE(GraphBuilder, FNNIOperatorGather);
+			RDG_EVENT_SCOPE(GraphBuilder, "NNE.Operator.Hlsl.Gather");
+			RDG_GPU_STAT_SCOPE(GraphBuilder, FNNEOperatorGather);
 
 			FComputeShaderUtils::AddPass(
 				GraphBuilder,
@@ -126,7 +128,7 @@ namespace UE::NNIRuntimeRDG::Private::Hlsl
 		}
 	};
 
-	bool ValidateGatherOperator(const UE::NNECore::FAttributeMap& AttributeMap, TConstArrayView<EMLTensorDataType> InputTypes, TConstArrayView<NNX::FSymbolicTensorShape> InputShapes)
+	bool ValidateGatherOperator(const NNECore::FAttributeMap& AttributeMap, TConstArrayView<ENNETensorDataType> InputTypes, TConstArrayView<NNECore::FSymbolicTensorShape> InputShapes)
 	{
 		bool bIsValid = true;
 
@@ -135,7 +137,7 @@ namespace UE::NNIRuntimeRDG::Private::Hlsl
 		bIsValid &= AttributeValidator.Validate(AttributeMap);
 
 		NNX::FInputValidator InputValidator;
-		InputValidator.AddSupportedType(EMLTensorDataType::Float);
+		InputValidator.AddSupportedType(ENNETensorDataType::Float);
 		InputValidator.AddRequired();
 		InputValidator.AddRequired();//Indices should be int32 or int64
 		bIsValid &= InputValidator.Validate(InputTypes);
diff --git a/Engine/Plugins/Experimental/NNE/Source/NNXRuntimeRDG/Private/Hlsl/NNIRuntimeRDGGemm.cpp b/Engine/Plugins/Experimental/NNE/Source/NNXRuntimeRDG/Private/Hlsl/NNIRuntimeRDGGemm.cpp
--- a/Engine/Plugins/Experimental/NNE/Source/NNXRuntimeRDG/Private/Hlsl/NNIRuntimeRDGGemm.cpp
+++ b/Engine/Plugins/Experimental/NNE/Source/NNXRuntimeRDG/Private/Hlsl/NNIRuntimeRDGGemm.cpp
@@ -4,10 +4,12 @@
 #include "NNEHlslShadersGemmCS.h"
 #include "NNXRuntimeHLSLHelper.h"
 #include "NNECoreAttributeMap.h"
+#include "NNECoreTensor.h"
+#include "NNECoreTypes.h"
 
 namespace UE::NNIRuntimeRDG::Private::Hlsl
 {
-	DECLARE_GPU_STAT_NAMED(FNNIOperatorGemm, TEXT("NNI.Operator.Hlsl.Gemm"));
+	DECLARE_GPU_STAT_NAMED(FNNEOperatorGemm, TEXT("NNE.Operator.Hlsl.Gemm"));
 
 	/**
 	 * Gemm operator implementation
@@ -31,7 +33,7 @@ namespace UE::NNIRuntimeRDG::Private::Hlsl
 
 	public:
 
-		virtual int PrepareOutputs(TConstArrayView<NNX::FTensorRef> InputTensors, TArrayView<NNX::FTensorRef> OutputTensors) const override
+		virtual int PrepareOutputs(TConstArrayView<NNECore::Internal::FTensorRef> InputTensors, TArrayView<NNECore::Internal::FTensorRef> OutputTensors) const override
 		{
 			check(InputTensors.Num() >= 2 && InputTensors.Num() <= 3);
 			check(OutputTensors.Num() == 1);
@@ -43,23 +45,26 @@ namespace UE::NNIRuntimeRDG::Private::Hlsl
 				return -1;
 			}
 
-			uint32 M = InputTransA != 0 ? InputA.Data[1] : InputA.Data[0];
-			uint32 N = InputTransB != 0 ? InputB.Data[0] : InputB.Data[1];
-			NNX::FTensorShape OutputShape;
+			uint32 M = InputTransA != 0 ? InputA.GetData()[1] : InputA.GetData()[0];
+			uint32 N = InputTransB != 0 ? InputB.GetData()[0] : InputB.GetData()[1];
+			TArray<uint32> OutputShapeData;
+
+			OutputShapeData.Emplace(M);
+			OutputShapeData.Emplace(N);
+			
+			NNECore::FTensorShape OutputShape = NNECore::FTensorShape::Make(OutputShapeData);
 			
-			OutputShape.Data.Emplace(M);
-			OutputShape.Data.Emplace(N);
 			OutputTensors[0]->SetShape(OutputShape);
 			return 0;
 		};
 		
-		virtual bool Initialize(TConstArrayView<NNX::FTensorDesc> InputTensorDescs, TConstArrayView<NNX::FTensorDesc> OutputTensorDescs, const UE::NNECore::FAttributeMap& Attributes) override
+		virtual bool Initialize(TConstArrayView<NNECore::FTensorDesc> InputTensorDescs, TConstArrayView<NNECore::FTensorDesc> OutputTensorDescs, const NNECore::FAttributeMap& Attributes) override
 		{
 			check(InputTensorDescs.Num() >= 2 && InputTensorDescs.Num() <= 3);
 			check(OutputTensorDescs.Num() == 1);
 
-            const NNX::FTensorDesc& InputA = InputTensorDescs[0];
-			const NNX::FTensorDesc& InputB = InputTensorDescs[1];
+            const NNECore::FTensorDesc& InputA = InputTensorDescs[0];
+			const NNECore::FTensorDesc& InputB = InputTensorDescs[1];
 
 			if (InputA.GetShape().Rank() != 2)
 			{
@@ -73,13 +78,13 @@ namespace UE::NNIRuntimeRDG::Private::Hlsl
 			}
 			if (InputTensorDescs.Num() == 3)
 			{
-				const NNX::FTensorDesc& InputC = InputTensorDescs[2];
+				const NNECore::FTensorDesc& InputC = InputTensorDescs[2];
 				if (InputC.GetShape().Rank() > 2)
 				{
 					UE_LOG(LogNNX, Warning, TEXT("Gemm third input should be of rank 2 or less"));
 					return false;
 				}
-				if (InputC.GetShape().Rank() == 1 && InputC.GetShape().Data[0] == 1)
+				if (InputC.GetShape().Rank() == 1 && InputC.GetShape().GetData()[0] == 1)
 				{
 					UE_LOG(LogNNX, Warning, TEXT("Gemm third input as scalar not supported"));
 					return false;
@@ -141,12 +146,12 @@ namespace UE::NNIRuntimeRDG::Private::Hlsl
 
 			FIntVector ThreadGroupCount = TGemmCS::GetGroupCount(*Parameters, Algorithm, 0);
 
-			RDG_EVENT_SCOPE(GraphBuilder, "NNI.Operator.Hlsl.Gemm");
-			RDG_GPU_STAT_SCOPE(GraphBuilder, FNNIOperatorGemm);
+			RDG_EVENT_SCOPE(GraphBuilder, "NNE.Operator.Hlsl.Gemm");
+			RDG_GPU_STAT_SCOPE(GraphBuilder, FNNEOperatorGemm);
 
 			FComputeShaderUtils::AddPass(
 				GraphBuilder,
-				RDG_EVENT_NAME("NNI.Operator.Hlsl.Dispatch"),
+				RDG_EVENT_NAME("NNE.Operator.Hlsl.Dispatch"),
 				ERDGPassFlags::Compute | ERDGPassFlags::NeverCull,
 				ComputeShader,
 				Parameters,
@@ -154,7 +159,7 @@ namespace UE::NNIRuntimeRDG::Private::Hlsl
 		}
 	};
 
-	bool ValidateGemmOperator(const UE::NNECore::FAttributeMap& AttributeMap, TConstArrayView<EMLTensorDataType> InputTypes, TConstArrayView<NNX::FSymbolicTensorShape> InputShapes)
+	bool ValidateGemmOperator(const NNECore::FAttributeMap& AttributeMap, TConstArrayView<ENNETensorDataType> InputTypes, TConstArrayView<NNECore::FSymbolicTensorShape> InputShapes)
 	{
 		bool bIsValid = true;
 
@@ -166,7 +171,7 @@ namespace UE::NNIRuntimeRDG::Private::Hlsl
 		bIsValid &= AttributeValidator.Validate(AttributeMap);
 
 		NNX::FInputValidator InputValidator;
-		InputValidator.AddSupportedType(EMLTensorDataType::Float);
+		InputValidator.AddSupportedType(ENNETensorDataType::Float);
 		InputValidator.AddRequired();
 		InputValidator.AddRequired();
 		InputValidator.AddOptional();
diff --git a/Engine/Plugins/Experimental/NNE/Source/NNXRuntimeRDG/Private/Hlsl/NNIRuntimeRDGInstanceNormalization.cpp b/Engine/Plugins/Experimental/NNE/Source/NNXRuntimeRDG/Private/Hlsl/NNIRuntimeRDGInstanceNormalization.cpp
--- a/Engine/Plugins/Experimental/NNE/Source/NNXRuntimeRDG/Private/Hlsl/NNIRuntimeRDGInstanceNormalization.cpp
+++ b/Engine/Plugins/Experimental/NNE/Source/NNXRuntimeRDG/Private/Hlsl/NNIRuntimeRDGInstanceNormalization.cpp
@@ -3,11 +3,13 @@
 #include "NNIRuntimeRDGInstanceNormalization.h"
 #include "NNEHlslShadersInstanceNormalizationCS.h"
 #include "NNXRuntimeHLSLHelper.h"
+#include "NNECoreTensor.h"
+#include "NNECoreTypes.h"
 
 namespace
 {
 
-int ValidateInput(const NNX::FTensorShape& Input, const NNX::FTensorShape& Scale, const NNX::FTensorShape& Bias)
+int ValidateInput(const UE::NNECore::FTensorShape& Input, const UE::NNECore::FTensorShape& Scale, const UE::NNECore::FTensorShape& Bias)
 {
 	if (Input.Rank() < 3)
 	{
@@ -20,15 +22,15 @@ int ValidateInput(const NNX::FTensorShape& Input, const NNX::FTensorShape& Scale
 		UE_LOG(LogNNX, Warning, TEXT("InstanceNormalization input scale should be of rank 1: %d"), Scale.Rank());
 		return -1;
 	}
-	if (Scale.Data[0] != Input.Data[1])
+	if (Scale.GetData()[0] != Input.GetData()[1])
 	{
-		UE_LOG(LogNNX, Warning, TEXT("InstanceNormalization input scale size should be equal to channel count: %d vs %d"), Scale.Data[0], Input.Data[1]);
+		UE_LOG(LogNNX, Warning, TEXT("InstanceNormalization input scale size should be equal to channel count: %d vs %d"), Scale.GetData()[0], Input.GetData()[1]);
 		return -1;
 	}
 
-	if (Bias.Data[0] != Input.Data[1])
+	if (Bias.GetData()[0] != Input.GetData()[1])
 	{
-		UE_LOG(LogNNX, Warning, TEXT("InstanceNormalization intput B size should be equal to channel count: : %d vs %d"), Bias.Data[0], Input.Data[1]);
+		UE_LOG(LogNNX, Warning, TEXT("InstanceNormalization intput B size should be equal to channel count: : %d vs %d"), Bias.GetData()[0], Input.GetData()[1]);
 		return -1;
 	}
 	if (Bias.Rank() != 1)
@@ -44,7 +46,7 @@ int ValidateInput(const NNX::FTensorShape& Input, const NNX::FTensorShape& Scale
 
 namespace UE::NNIRuntimeRDG::Private::Hlsl
 {
-	DECLARE_GPU_STAT_NAMED(FNNIOperatorInstanceNormalization, TEXT("NNI.Operator.Hlsl.InstanceNormalization"));
+	DECLARE_GPU_STAT_NAMED(FNNEOperatorInstanceNormalization, TEXT("NNE.Operator.Hlsl.InstanceNormalization"));
 
 	using EInstanceNormalizationAlgorithm = UE::NNEHlslShaders::Internal::EInstanceNormalizationAlgorithm;
 
@@ -66,7 +68,7 @@ namespace UE::NNIRuntimeRDG::Private::Hlsl
 
 	public:
 
-		virtual int PrepareOutputs(TConstArrayView<NNX::FTensorRef> InputTensors, TArrayView<NNX::FTensorRef> OutputTensors) const override
+		virtual int PrepareOutputs(TConstArrayView<NNECore::Internal::FTensorRef> InputTensors, TArrayView<NNECore::Internal::FTensorRef> OutputTensors) const override
 		{
 			check(InputTensors.Num() == 3);
 			check(OutputTensors.Num() == 1);
@@ -81,7 +83,7 @@ namespace UE::NNIRuntimeRDG::Private::Hlsl
 			return 0;
 		};
 
-		virtual bool Initialize(TConstArrayView<NNX::FTensorDesc> InputTensorDescs, TConstArrayView<NNX::FTensorDesc> OutputTensorDescs, const UE::NNECore::FAttributeMap& Attributes) override
+		virtual bool Initialize(TConstArrayView<NNECore::FTensorDesc> InputTensorDescs, TConstArrayView<NNECore::FTensorDesc> OutputTensorDescs, const UE::NNECore::FAttributeMap& Attributes) override
 		{
 			using namespace UE::NNEHlslShaders::Internal;
 
@@ -131,8 +133,8 @@ namespace UE::NNIRuntimeRDG::Private::Hlsl
 			TShaderMapRef<TInstanceNormalizationCS> ComputeShader(GetGlobalShaderMap(GMaxRHIFeatureLevel), PermutationVector);
 			FIntVector ThreadGroupCount = TInstanceNormalizationCS::GetGroupCount(*Parameters, DispatchAlgorithm);
 
-			RDG_EVENT_SCOPE(GraphBuilder, "NNI.Operator.Hlsl.InstanceNormalization");
-			RDG_GPU_STAT_SCOPE(GraphBuilder, FNNIOperatorInstanceNormalization);
+			RDG_EVENT_SCOPE(GraphBuilder, "NNE.Operator.Hlsl.InstanceNormalization");
+			RDG_GPU_STAT_SCOPE(GraphBuilder, FNNEOperatorInstanceNormalization);
 
 			FComputeShaderUtils::AddPass(
 				GraphBuilder,
@@ -144,7 +146,7 @@ namespace UE::NNIRuntimeRDG::Private::Hlsl
 		}
 	};
 
-	bool ValidateInstanceNormalizationOperator(const UE::NNECore::FAttributeMap& AttributeMap, TConstArrayView<EMLTensorDataType> InputTypes, TConstArrayView<NNX::FSymbolicTensorShape> InputShapes)
+	bool ValidateInstanceNormalizationOperator(const NNECore::FAttributeMap& AttributeMap, TConstArrayView<ENNETensorDataType> InputTypes, TConstArrayView<NNECore::FSymbolicTensorShape> InputShapes)
 	{
 		bool bIsValid = true;
 
@@ -154,7 +156,7 @@ namespace UE::NNIRuntimeRDG::Private::Hlsl
 		bIsValid &= AttributeValidator.Validate(AttributeMap);
 
 		NNX::FInputValidator InputValidator;
-		InputValidator.AddSupportedType(EMLTensorDataType::Float);
+		InputValidator.AddSupportedType(ENNETensorDataType::Float);
 		InputValidator.AddRequired();
 		InputValidator.AddRequired();
 		InputValidator.AddRequired();
diff --git a/Engine/Plugins/Experimental/NNE/Source/NNXRuntimeRDG/Private/Hlsl/NNIRuntimeRDGMatMul.cpp b/Engine/Plugins/Experimental/NNE/Source/NNXRuntimeRDG/Private/Hlsl/NNIRuntimeRDGMatMul.cpp
--- a/Engine/Plugins/Experimental/NNE/Source/NNXRuntimeRDG/Private/Hlsl/NNIRuntimeRDGMatMul.cpp
+++ b/Engine/Plugins/Experimental/NNE/Source/NNXRuntimeRDG/Private/Hlsl/NNIRuntimeRDGMatMul.cpp
@@ -3,10 +3,12 @@
 #include "NNIRuntimeRDGMatMul.h"
 #include "NNEHlslShadersGemmCS.h"
 #include "NNXRuntimeHLSLHelper.h"
+#include "NNECoreTensor.h"
+#include "NNECoreTypes.h"
 
 namespace UE::NNIRuntimeRDG::Private::Hlsl
 {
-    DECLARE_GPU_STAT_NAMED(FNNIOperatorMatMul, TEXT("NNI.Operator.Hlsl.MatMul"));
+    DECLARE_GPU_STAT_NAMED(FNNEOperatorMatMul, TEXT("NNE.Operator.Hlsl.MatMul"));
 
 	/**
 	 * MatMul operator implementation
@@ -20,7 +22,7 @@ namespace UE::NNIRuntimeRDG::Private::Hlsl
 
 	public:
 
-		virtual int PrepareOutputs(TConstArrayView<NNX::FTensorRef> InputTensors, TArrayView<NNX::FTensorRef> OutputTensors) const override
+		virtual int PrepareOutputs(TConstArrayView<NNECore::Internal::FTensorRef> InputTensors, TArrayView<NNECore::Internal::FTensorRef> OutputTensors) const override
 		{
 			check(InputTensors.Num() == 2);
 			check(OutputTensors.Num() == 1);
@@ -38,44 +40,45 @@ namespace UE::NNIRuntimeRDG::Private::Hlsl
 				UE_LOG(LogNNX, Warning, TEXT("Matmul second input should be at least of rank 2"));
 				return -1;
 			}
-			if (InputA.Data[InputA.Rank() - 1] != InputB.Data[InputB.Rank() - 2])
+			if (InputA.GetData()[InputA.Rank() - 1] != InputB.GetData()[InputB.Rank() - 2])
 			{
 				UE_LOG(LogNNX, Warning, TEXT("Matmul first input last dimension should be equal to second input last dimension"));
 				return -1;
 			}
 
 			const int32 OutputRank = FMath::Max(InputA.Rank(), InputB.Rank());
-			NNX::FTensorShape OutputShape;
-			
-			OutputShape.Data.SetNumUninitialized(OutputRank);
+			TArray<uint32> OutputShapeData;
+			OutputShapeData.SetNumUninitialized(OutputRank);
 			
 			//Broadcast
 			for (int32 i = 0; i < OutputRank; ++i)
 			{
 				int32 AIndex = InputA.Rank() - 1 - i;
 				int32 BIndex = InputB.Rank() - 1 - i;
-				int32 AValue = AIndex >= 0 ? InputA.Data[AIndex] : 1;
-				int32 BValue = BIndex >= 0 ? InputB.Data[BIndex] : 1;
+				int32 AValue = AIndex >= 0 ? InputA.GetData()[AIndex] : 1;
+				int32 BValue = BIndex >= 0 ? InputB.GetData()[BIndex] : 1;
 				int32 OutputValue = FMath::Max(AValue, BValue);
-				OutputShape.Data[OutputRank - 1 - i] = OutputValue;
+				OutputShapeData[OutputRank - 1 - i] = OutputValue;
 			}
 			
 			//2D Mat
-			OutputShape.Data[OutputRank - 2] = InputA.Data[InputA.Rank() - 2];
-			OutputShape.Data[OutputRank - 1] = InputB.Data[InputB.Rank() - 1];
+			OutputShapeData[OutputRank - 2] = InputA.GetData()[InputA.Rank() - 2];
+			OutputShapeData[OutputRank - 1] = InputB.GetData()[InputB.Rank() - 1];
+
+			NNECore::FTensorShape OutputShape = NNECore::FTensorShape::Make(OutputShapeData);
 
 			OutputTensors[0]->SetShape(OutputShape);
 
 			return 0;
 		};
 
-		virtual bool Initialize(TConstArrayView<NNX::FTensorDesc> InputTensorDescs, TConstArrayView<NNX::FTensorDesc> OutputTensorDescs, const UE::NNECore::FAttributeMap& Attributes) override
+		virtual bool Initialize(TConstArrayView<NNECore::FTensorDesc> InputTensorDescs, TConstArrayView<NNECore::FTensorDesc> OutputTensorDescs, const NNECore::FAttributeMap& Attributes) override
 		{
 			check(InputTensorDescs.Num() == 2);
 			check(OutputTensorDescs.Num() == 1);
 
-			const NNX::FTensorDesc& InputA = InputTensorDescs[0];
-			const NNX::FTensorDesc& InputB = InputTensorDescs[1];
+			const NNECore::FTensorDesc& InputA = InputTensorDescs[0];
+			const NNECore::FTensorDesc& InputB = InputTensorDescs[1];
 
 			if (InputA.GetShape().Rank() < 2)
 			{
@@ -123,12 +126,12 @@ namespace UE::NNIRuntimeRDG::Private::Hlsl
 
 			FIntVector ThreadGroupCount = TGemmCS::GetGroupCount(*Parameters, Algorithm, NumStackDimensions);
 
-			RDG_EVENT_SCOPE(GraphBuilder, "NNI.Operator.Hlsl.MatMul");
-			RDG_GPU_STAT_SCOPE(GraphBuilder, FNNIOperatorMatMul);
+			RDG_EVENT_SCOPE(GraphBuilder, "NNE.Operator.Hlsl.MatMul");
+			RDG_GPU_STAT_SCOPE(GraphBuilder, FNNEOperatorMatMul);
 
 			FComputeShaderUtils::AddPass(
 				GraphBuilder,
-				RDG_EVENT_NAME("NNI.Operator.Hlsl.MatMul.Dispatch"),
+				RDG_EVENT_NAME("NNE.Operator.Hlsl.MatMul.Dispatch"),
 				ERDGPassFlags::Compute | ERDGPassFlags::NeverCull,
 				ComputeShader,
 				Parameters,
@@ -136,7 +139,7 @@ namespace UE::NNIRuntimeRDG::Private::Hlsl
 		}
 	};
 
-	bool ValidateMatMulOperator(const UE::NNECore::FAttributeMap& AttributeMap, TConstArrayView<EMLTensorDataType> InputTypes, TConstArrayView<NNX::FSymbolicTensorShape> InputShapes)
+	bool ValidateMatMulOperator(const NNECore::FAttributeMap& AttributeMap, TConstArrayView<ENNETensorDataType> InputTypes, TConstArrayView<NNECore::FSymbolicTensorShape> InputShapes)
 	{
 		bool bIsValid = true;
 
@@ -144,7 +147,7 @@ namespace UE::NNIRuntimeRDG::Private::Hlsl
 		bIsValid &= AttributeValidator.Validate(AttributeMap);
 
 		NNX::FInputValidator InputValidator;
-		InputValidator.AddSupportedType(EMLTensorDataType::Float);
+		InputValidator.AddSupportedType(ENNETensorDataType::Float);
 		InputValidator.AddRequired();
 		InputValidator.AddRequired();
 		bIsValid &= InputValidator.Validate(InputTypes);
diff --git a/Engine/Plugins/Experimental/NNE/Source/NNXRuntimeRDG/Private/Hlsl/NNIRuntimeRDGPad.cpp b/Engine/Plugins/Experimental/NNE/Source/NNXRuntimeRDG/Private/Hlsl/NNIRuntimeRDGPad.cpp
--- a/Engine/Plugins/Experimental/NNE/Source/NNXRuntimeRDG/Private/Hlsl/NNIRuntimeRDGPad.cpp
+++ b/Engine/Plugins/Experimental/NNE/Source/NNXRuntimeRDG/Private/Hlsl/NNIRuntimeRDGPad.cpp
@@ -3,10 +3,12 @@
 #include "NNIRuntimeRDGPad.h"
 #include "NNEHlslShadersPadCS.h"
 #include "NNXRuntimeHLSLHelper.h"
+#include "NNECoreTensor.h"
+#include "NNECoreTypes.h"
 
 namespace UE::NNIRuntimeRDG::Private::Hlsl
 {
-	DECLARE_GPU_STAT_NAMED(FNNIOperatorPad, TEXT("NNI.Operator.Hlsl.Pad"));
+	DECLARE_GPU_STAT_NAMED(FNNEOperatorPad, TEXT("NNE.Operator.Hlsl.Pad"));
 
 	using EPadMode = UE::NNEHlslShaders::Internal::EPadMode;
 
@@ -26,33 +28,34 @@ namespace UE::NNIRuntimeRDG::Private::Hlsl
 
 	public:
 
-		virtual int PrepareOutputs(TConstArrayView<NNX::FTensorRef> InputTensors, TArrayView<NNX::FTensorRef> OutputTensors) const override
+		virtual int PrepareOutputs(TConstArrayView<NNECore::Internal::FTensorRef> InputTensors, TArrayView<NNECore::Internal::FTensorRef> OutputTensors) const override
 		{
 			check(InputTensors.Num() == 1);
 			check(OutputTensors.Num() == 1);
 
-			const NNX::FTensor& X = *InputTensors[0];
+			const NNECore::Internal::FTensor& X = *InputTensors[0];
 
-			NNX::FTensorShape OutputShape;
+			TArray<uint32> OutputShapeData;
 			for (int32 i = 0; i < X.GetShape().Rank(); ++i)
 			{
 				int32 PrePad = Pads[i];
 				int32 PostPad = Pads[i + X.GetShape().Rank()];
-				int32 OutputDim = PrePad + X.GetShape().Data[i] + PostPad;
+				int32 OutputDim = PrePad + X.GetShape().GetData()[i] + PostPad;
 				if (OutputDim < 1)
 				{
-					UE_LOG(LogNNX, Warning, TEXT("Pads cannot reduce dimension below 1, but would for tensor (name:%s) at rank %d of size %d with prepad %d and postpad %d."), *X.GetName(), i, X.GetShape().Data[i], PrePad, PostPad);
+					UE_LOG(LogNNX, Warning, TEXT("Pads cannot reduce dimension below 1, but would for tensor (name:%s) at rank %d of size %d with prepad %d and postpad %d."), *X.GetName(), i, X.GetShape().GetData()[i], PrePad, PostPad);
 					return -1;
 				}
-				OutputShape.Data.Emplace(OutputDim);
+				OutputShapeData.Emplace(OutputDim);
 			}
 
+			NNECore::FTensorShape OutputShape = NNECore::FTensorShape::Make(OutputShapeData);
 			OutputTensors[0]->SetShape(OutputShape);
 
 			return 0;
 		};
 
-		virtual bool Initialize(TConstArrayView<NNX::FTensorDesc> InputTensorDescs, TConstArrayView<NNX::FTensorDesc> OutputTensorDescs, const UE::NNECore::FAttributeMap& Attributes) override
+		virtual bool Initialize(TConstArrayView<NNECore::FTensorDesc> InputTensorDescs, TConstArrayView<NNECore::FTensorDesc> OutputTensorDescs, const NNECore::FAttributeMap& Attributes) override
 		{
 			using namespace UE::NNEHlslShaders::Internal;
 			
@@ -109,12 +112,12 @@ namespace UE::NNIRuntimeRDG::Private::Hlsl
 
 			TShaderMapRef<FPadCS> ComputeShader(GetGlobalShaderMap(GMaxRHIFeatureLevel), PermutationVector);
 
-			RDG_EVENT_SCOPE(GraphBuilder, "NNI.Operator.Hlsl.Pad");
-			RDG_GPU_STAT_SCOPE(GraphBuilder, FNNIOperatorPad);
+			RDG_EVENT_SCOPE(GraphBuilder, "NNE.Operator.Hlsl.Pad");
+			RDG_GPU_STAT_SCOPE(GraphBuilder, FNNEOperatorPad);
 
 			FComputeShaderUtils::AddPass(
 				GraphBuilder,
-				RDG_EVENT_NAME("NNI.Operator.Hlsl.Pad.Dispatch"),
+				RDG_EVENT_NAME("NNE.Operator.Hlsl.Pad.Dispatch"),
 				ERDGPassFlags::Compute | ERDGPassFlags::NeverCull,
 				ComputeShader,
 				Params,
@@ -122,7 +125,7 @@ namespace UE::NNIRuntimeRDG::Private::Hlsl
 		}
 	};
 
-	bool ValidatePadOperator(const UE::NNECore::FAttributeMap& AttributeMap, TConstArrayView<EMLTensorDataType> InputTypes, TConstArrayView<NNX::FSymbolicTensorShape> InputShapes)
+	bool ValidatePadOperator(const NNECore::FAttributeMap& AttributeMap, TConstArrayView<ENNETensorDataType> InputTypes, TConstArrayView<NNECore::FSymbolicTensorShape> InputShapes)
 	{
 		bool bIsValid = true;
 
@@ -144,7 +147,7 @@ namespace UE::NNIRuntimeRDG::Private::Hlsl
 		}
 		
 		NNX::FInputValidator InputValidator;
-		InputValidator.AddSupportedType(EMLTensorDataType::Float);
+		InputValidator.AddSupportedType(ENNETensorDataType::Float);
 		InputValidator.AddRequired();
 		bIsValid &= InputValidator.Validate(InputTypes);
 
diff --git a/Engine/Plugins/Experimental/NNE/Source/NNXRuntimeRDG/Private/Hlsl/NNIRuntimeRDGUpsample.cpp b/Engine/Plugins/Experimental/NNE/Source/NNXRuntimeRDG/Private/Hlsl/NNIRuntimeRDGUpsample.cpp
--- a/Engine/Plugins/Experimental/NNE/Source/NNXRuntimeRDG/Private/Hlsl/NNIRuntimeRDGUpsample.cpp
+++ b/Engine/Plugins/Experimental/NNE/Source/NNXRuntimeRDG/Private/Hlsl/NNIRuntimeRDGUpsample.cpp
@@ -3,10 +3,12 @@
 #include "NNIRuntimeRDGUpsample.h"
 #include "NNEHlslShadersUpsampleCS.h"
 #include "NNXRuntimeHLSLHelper.h"
+#include "NNECoreTypes.h"
+#include "NNECoreTensor.h"
 
 namespace UE::NNIRuntimeRDG::Private::Hlsl
 {
-	DECLARE_GPU_STAT_NAMED(FNNIOperatorUpsample, TEXT("NNI.Operator.Hlsl.Upsample"));
+	DECLARE_GPU_STAT_NAMED(FNNEOperatorUpsample, TEXT("NNE.Operator.Hlsl.Upsample"));
 
 	/**
 	 * Upsample operator implementation
@@ -20,13 +22,13 @@ namespace UE::NNIRuntimeRDG::Private::Hlsl
 
 	public:
 
-		virtual int PrepareOutputs(TConstArrayView<NNX::FTensorRef> InputTensors, TArrayView<NNX::FTensorRef> OutputTensors) const override
+		virtual int PrepareOutputs(TConstArrayView<NNECore::Internal::FTensorRef> InputTensors, TArrayView<NNECore::Internal::FTensorRef> OutputTensors) const override
 		{
 			check(InputTensors.Num() == 2);
 			check(OutputTensors.Num() == 1);
 
-			const NNX::FTensor& X = *InputTensors[0];
-			const NNX::FTensor& Scales = *InputTensors[1];
+			const NNECore::Internal::FTensor& X = *InputTensors[0];
+			const NNECore::Internal::FTensor& Scales = *InputTensors[1];
 			
 			if (!Scales.HasPreparedData())
 			{
@@ -42,13 +44,14 @@ namespace UE::NNIRuntimeRDG::Private::Hlsl
 				return -1;
 			}
 
-			NNX::FTensorShape OutputShape;
+			TArray<uint32> OutputShapeData;
 			for (int32 i = 0; i < X.GetShape().Rank(); ++i)
 			{
 				
-				OutputShape.Data.Emplace(FMath::FloorToInt32(X.GetShape().Data[i] * ScalesData[i]));
+				OutputShapeData.Emplace(FMath::FloorToInt32(X.GetShape().GetData()[i] * ScalesData[i]));
 			}
 
+			NNECore::FTensorShape OutputShape = NNECore::FTensorShape::Make(OutputShapeData);
 			OutputTensors[0]->SetShape(OutputShape);
 
 			return 0;
@@ -100,12 +103,12 @@ namespace UE::NNIRuntimeRDG::Private::Hlsl
 
 			TShaderMapRef<FUpsampleCS> ComputeShader(GetGlobalShaderMap(GMaxRHIFeatureLevel), PermutationVector);
 
-			RDG_EVENT_SCOPE(GraphBuilder, "NNI.Operator.Hlsl.Upsample");
-			RDG_GPU_STAT_SCOPE(GraphBuilder, FNNIOperatorUpsample);
+			RDG_EVENT_SCOPE(GraphBuilder, "NNE.Operator.Hlsl.Upsample");
+			RDG_GPU_STAT_SCOPE(GraphBuilder, FNNEOperatorUpsample);
 
 			FComputeShaderUtils::AddPass(
 				GraphBuilder,
-				RDG_EVENT_NAME("NNI.Operator.Hlsl.Upsample.Dispatch"),
+				RDG_EVENT_NAME("NNE.Operator.Hlsl.Upsample.Dispatch"),
 				ERDGPassFlags::Compute | ERDGPassFlags::NeverCull,
 				ComputeShader,
 				Params,
@@ -113,7 +116,7 @@ namespace UE::NNIRuntimeRDG::Private::Hlsl
 		}
 	};
 
-	bool ValidateUpsampleOperator(const UE::NNECore::FAttributeMap& AttributeMap, TConstArrayView<EMLTensorDataType> InputTypes, TConstArrayView<NNX::FSymbolicTensorShape> InputShapes)
+	bool ValidateUpsampleOperator(const NNECore::FAttributeMap& AttributeMap, TConstArrayView<ENNETensorDataType> InputTypes, TConstArrayView<NNECore::FSymbolicTensorShape> InputShapes)
 	{
 		bool bIsValid = true;
 
@@ -129,7 +132,7 @@ namespace UE::NNIRuntimeRDG::Private::Hlsl
 		}
 
 		NNX::FInputValidator InputValidator;
-		InputValidator.AddSupportedType(EMLTensorDataType::Float);
+		InputValidator.AddSupportedType(ENNETensorDataType::Float);
 		InputValidator.AddRequired();
 		InputValidator.AddRequired();
 		bIsValid &= InputValidator.Validate(InputTypes);
diff --git a/Engine/Plugins/Experimental/NNE/Source/NNXRuntimeRDG/Private/NNERuntimeRDGModelHLSL.cpp b/Engine/Plugins/Experimental/NNE/Source/NNXRuntimeRDG/Private/NNERuntimeRDGModelHLSL.cpp
--- a/Engine/Plugins/Experimental/NNE/Source/NNXRuntimeRDG/Private/NNERuntimeRDGModelHLSL.cpp
+++ b/Engine/Plugins/Experimental/NNE/Source/NNXRuntimeRDG/Private/NNERuntimeRDGModelHLSL.cpp
@@ -1,6 +1,7 @@
 // Copyright Epic Games, Inc. All Rights Reserved.
 
 #include "NNERuntimeRDGModelHLSL.h"
+#include "NNECoreTensor.h"
 #include "NNXRuntimeHLSLOp.h"
 #include "RenderGraphBuilder.h"
 #include "RenderGraphUtils.h"
@@ -319,8 +320,8 @@ int FModel::PrepareTensorShapesAndData()
 	// Run model preparation (including shape inference) on all operators
 	// This loop could be abstracted to a different engine/system as it apply on FTensorRef & IPrepareOperator witch are RDG agnostics.
 	static constexpr int32 MaxExpectedInput = 10;
-	TArray<NNX::FTensorRef, TInlineAllocator<MaxExpectedInput>> InputTensors;
-	TArray<NNX::FTensorRef> OutputTensors;
+	TArray<NNECore::Internal::FTensorRef, TInlineAllocator<MaxExpectedInput>> InputTensors;
+	TArray<NNECore::Internal::FTensorRef> OutputTensors;
 	TArray<bool> AllInitializedTensors;
 
 	checkCode(
diff --git a/Engine/Plugins/Experimental/NNE/Source/NNXRuntimeRDG/Private/NNERuntimeRDGModelHLSL.h b/Engine/Plugins/Experimental/NNE/Source/NNXRuntimeRDG/Private/NNERuntimeRDGModelHLSL.h
--- a/Engine/Plugins/Experimental/NNE/Source/NNXRuntimeRDG/Private/NNERuntimeRDGModelHLSL.h
+++ b/Engine/Plugins/Experimental/NNE/Source/NNXRuntimeRDG/Private/NNERuntimeRDGModelHLSL.h
@@ -3,7 +3,7 @@
 #pragma once
 
 #include "NNXRuntimeRDG.h"
-#include "NNXTypes.h"
+#include "NNECoreTypes.h"
 #include "NNERuntimeRDGTensorHLSL.h"
 
 class FRDGBuilder;
diff --git a/Engine/Plugins/Experimental/NNE/Source/NNXRuntimeRDG/Private/NNERuntimeRDGNeuralPostProcessing.cpp b/Engine/Plugins/Experimental/NNE/Source/NNXRuntimeRDG/Private/NNERuntimeRDGNeuralPostProcessing.cpp
--- a/Engine/Plugins/Experimental/NNE/Source/NNXRuntimeRDG/Private/NNERuntimeRDGNeuralPostProcessing.cpp
+++ b/Engine/Plugins/Experimental/NNE/Source/NNXRuntimeRDG/Private/NNERuntimeRDGNeuralPostProcessing.cpp
@@ -253,14 +253,14 @@ void FNNENeuralPostProcessing::PrePostProcessPass_RenderThread(FRDGBuilder& Grap
 						InputSize = *InputSizes.Find(Pair.Key);
 					}
 
-					NNX::FSymbolicTensorShape InputShape = Pair.Value->GetInputTensorDescs()[0].GetShape();
+					UE::NNECore::FSymbolicTensorShape InputShape = Pair.Value->GetInputTensorDescs()[0].GetShape();
 
 					checkf(InputShape.Rank() == 4, TEXT("Neural Post Processing requires models with input shape 1 x 3 x height x width!"))
-					checkf(InputShape.Data[0] == 1, TEXT("Neural Post Processing requires models with input shape 1 x 3 x height x width!"))
-					checkf(InputShape.Data[1] == 3, TEXT("Neural Post Processing requires models with input shape 1 x 3 x height x width!"))
+					checkf(InputShape.GetData()[0] == 1, TEXT("Neural Post Processing requires models with input shape 1 x 3 x height x width!"))
+					checkf(InputShape.GetData()[1] == 3, TEXT("Neural Post Processing requires models with input shape 1 x 3 x height x width!"))
 
-					int32 NeuralNetworkInputWidth = InputShape.Data[3] >= 0 ? InputShape.Data[3] : (InputSize.X >= 0 ? InputSize.X : TextureSize.X);
-					int32 NeuralNetworkInputHeight = InputShape.Data[2] >= 0 ? InputShape.Data[2] : (InputSize.Y >= 0 ? InputSize.Y : TextureSize.Y);
+					int32 NeuralNetworkInputWidth = InputShape.GetData()[3] >= 0 ? InputShape.GetData()[3] : (InputSize.X >= 0 ? InputSize.X : TextureSize.X);
+					int32 NeuralNetworkInputHeight = InputShape.GetData()[2] >= 0 ? InputShape.GetData()[2] : (InputSize.Y >= 0 ? InputSize.Y : TextureSize.Y);
 
 					FRDGBufferDesc InputBufferDesc = FRDGBufferDesc::CreateBufferDesc(sizeof(float), NeuralNetworkInputWidth * NeuralNetworkInputHeight * 3);
 					FRDGBufferRef InputBuffer = GraphBuilder.CreateBuffer(InputBufferDesc, *(FString("NNENeuralPostProcessing::NeuralNetowrkInput_") + FString::FromInt(Pair.Key)));
@@ -297,13 +297,13 @@ void FNNENeuralPostProcessing::PrePostProcessPass_RenderThread(FRDGBuilder& Grap
 					NNX::FTensorShape OutputShape = Pair.Value->GetOutputTensorShapes()[0];
 
 					checkf(OutputShape.Rank() == 4, TEXT("Neural Post Processing requires models with output shape 1 x 3 x height x width!"))
-					checkf(OutputShape.Data[0] == 1, TEXT("Neural Post Processing requires models with output shape 1 x 3 x height x width!"))
-					checkf(OutputShape.Data[1] == 3, TEXT("Neural Post Processing requires models with output shape 1 x 3 x height x width!"))
-					checkf(OutputShape.Data[2] > 0, TEXT("Neural Post Processing requires models with output height > 0!"))
-					checkf(OutputShape.Data[3] > 0, TEXT("Neural Post Processing requires models with output width > 0!"))
+					checkf(OutputShape.GetData()[0] == 1, TEXT("Neural Post Processing requires models with output shape 1 x 3 x height x width!"))
+					checkf(OutputShape.GetData()[1] == 3, TEXT("Neural Post Processing requires models with output shape 1 x 3 x height x width!"))
+					checkf(OutputShape.GetData()[2] > 0, TEXT("Neural Post Processing requires models with output height > 0!"))
+					checkf(OutputShape.GetData()[3] > 0, TEXT("Neural Post Processing requires models with output width > 0!"))
 
-					int32 NeuralNetworkOutputWidth = OutputShape.Data[3];
-					int32 NeuralNetworkOutputHeight = OutputShape.Data[2];
+					int32 NeuralNetworkOutputWidth = OutputShape.GetData()[3];
+					int32 NeuralNetworkOutputHeight = OutputShape.GetData()[2];
 
 					FRDGBufferDesc OutputBufferDesc = FRDGBufferDesc::CreateBufferDesc(sizeof(float), NeuralNetworkOutputWidth * NeuralNetworkOutputHeight * 3);
 					FRDGBufferRef OutputBuffer = GraphBuilder.CreateBuffer(OutputBufferDesc, *(FString("NNENeuralPostProcessing::NeuralNetowrkOutput_") + FString::FromInt(Pair.Key)));
diff --git a/Engine/Plugins/Experimental/NNE/Source/NNXRuntimeRDG/Private/NNERuntimeRDGTensorHLSL.h b/Engine/Plugins/Experimental/NNE/Source/NNXRuntimeRDG/Private/NNERuntimeRDGTensorHLSL.h
--- a/Engine/Plugins/Experimental/NNE/Source/NNXRuntimeRDG/Private/NNERuntimeRDGTensorHLSL.h
+++ b/Engine/Plugins/Experimental/NNE/Source/NNXRuntimeRDG/Private/NNERuntimeRDGTensorHLSL.h
@@ -3,7 +3,7 @@
 #pragma once
 
 #include "NNXRuntimeRDG.h"
-#include "NNXTypes.h"
+#include "NNECoreTypes.h"
 #include "RHIGPUReadback.h"
 
 using FCpuBufferRef = void*;
@@ -33,13 +33,13 @@ public:
 		this->PreparedData = TensorRDG.GetPreparedData<uint8>();//This copy the data, we should find another way
 	}
 	
-	FTensorHLSL(const FString& Name, EMLTensorDataType DataType, const NNX::FTensorShape& Shape)
+	FTensorHLSL(const FString& Name, ENNETensorDataType DataType, const NNX::FTensorShape& Shape)
 	{
 		this->Name = Name;
 		this->DataType = DataType;
 		this->Shape = Shape;
 		Volume = Shape.Volume();
-		DataSize = (uint64)NNX::GetTensorDataTypeSizeInBytes(DataType) * Volume;
+		DataSize = (uint64)UE::NNECore::GetTensorDataTypeSizeInBytes(DataType) * Volume;
 	}
 
 	bool HasUploadBuffer() const { return UploadBuffer != FCpuBufferRef{}; }
diff --git a/Engine/Plugins/Experimental/NNE/Source/NNXRuntimeRDG/Private/NNXRuntimeDML.cpp b/Engine/Plugins/Experimental/NNE/Source/NNXRuntimeRDG/Private/NNXRuntimeDML.cpp
--- a/Engine/Plugins/Experimental/NNE/Source/NNXRuntimeRDG/Private/NNXRuntimeDML.cpp
+++ b/Engine/Plugins/Experimental/NNE/Source/NNXRuntimeRDG/Private/NNXRuntimeDML.cpp
@@ -169,7 +169,7 @@ namespace DmlUtil
 		TArray<uint32, TInlineAllocator<FTensorShape::MaxRank>>		Strides;
 	};
 
-	void SetTensorStrides(FTensorDesc& TensorDesc, const FTensor& InputDesc)
+	void SetTensorStrides(FTensorDesc& TensorDesc, const UE::NNECore::Internal::FTensor& InputDesc)
 	{
 		uint32 CurrStride = 1;
 
@@ -178,11 +178,11 @@ namespace DmlUtil
 		for (int32 i = InputDesc.GetShape().Rank() - 1; i >= 0; --i)
 		{
 			TensorDesc.Strides[i] = CurrStride;
-			CurrStride *= InputDesc.GetShape().Data[i];
+			CurrStride *= InputDesc.GetShape().GetData()[i];
 		}
 	}
 
-	void SetTensorSizesAndStridesForBroadcast(FTensorDesc& TensorDesc, const FTensor& InputDesc, const FTensor& TargetDesc)
+	void SetTensorSizesAndStridesForBroadcast(FTensorDesc& TensorDesc, const UE::NNECore::Internal::FTensor& InputDesc, const UE::NNECore::Internal::FTensor& TargetDesc)
 	{
 		static_assert(FTensorShape::MaxRank <= 8);
 		
@@ -196,23 +196,23 @@ namespace DmlUtil
 		
 		for (int32 i = 0; i < (int32) TargetDimension; ++i)
 		{
-			TensorDesc.Sizes[i] = i < DimensionOffset ? 1 : InputDesc.GetShape().Data[i - DimensionOffset];
+			TensorDesc.Sizes[i] = i < DimensionOffset ? 1 : InputDesc.GetShape().GetData()[i - DimensionOffset];
 		}
 
 		uint32 CurrStride = 1;
 
 		for (int32 i = TargetDimension - 1; i >= 0; --i)
 		{
-			const bool bBroadcast = TensorDesc.Sizes[i] < TargetDesc.GetShape().Data[i];
+			const bool bBroadcast = TensorDesc.Sizes[i] < TargetDesc.GetShape().GetData()[i];
 
 			TensorDesc.Strides[i] = bBroadcast ? 0 : CurrStride;
 			CurrStride *= TensorDesc.Sizes[i];
 
-			TensorDesc.Sizes[i] = TargetDesc.GetShape().Data[i];
+			TensorDesc.Sizes[i] = TargetDesc.GetShape().GetData()[i];
 		}
 	}
 
-	inline bool IsSameShape(const FTensor& Left, const FTensor& Right)
+	inline bool IsSameShape(const UE::NNECore::Internal::FTensor& Left, const UE::NNECore::Internal::FTensor& Right)
 	{
 		if (Left.GetShape().Rank() != Right.GetShape().Rank())
 		{
@@ -221,7 +221,7 @@ namespace DmlUtil
 		
 		for (int32 Idx = 0; Idx < Left.GetShape().Rank(); ++Idx)
 		{
-			if (Left.GetShape().Data[Idx] != Right.GetShape().Data[Idx])
+			if (Left.GetShape().GetData()[Idx] != Right.GetShape().GetData()[Idx])
 			{
 				return false;
 			}
@@ -230,41 +230,41 @@ namespace DmlUtil
 		return true;
 	}
 
-	DML_TENSOR_DATA_TYPE GetTensorDataType(EMLTensorDataType DataType)
+	DML_TENSOR_DATA_TYPE GetTensorDataType(ENNETensorDataType DataType)
 	{
 		switch (DataType)
 		{
-		case EMLTensorDataType::Double:
+		case ENNETensorDataType::Double:
 			return DML_TENSOR_DATA_TYPE_FLOAT64;
 
-		case EMLTensorDataType::Float:
+		case ENNETensorDataType::Float:
 			return DML_TENSOR_DATA_TYPE_FLOAT32;
 
-		case EMLTensorDataType::Half:
+		case ENNETensorDataType::Half:
 			return DML_TENSOR_DATA_TYPE_FLOAT16;
 
-		case EMLTensorDataType::UInt64:
+		case ENNETensorDataType::UInt64:
 			return DML_TENSOR_DATA_TYPE_UINT64;
 
-		case EMLTensorDataType::UInt32:
+		case ENNETensorDataType::UInt32:
 			return DML_TENSOR_DATA_TYPE_UINT32;
 
-		case EMLTensorDataType::UInt16:
+		case ENNETensorDataType::UInt16:
 			return DML_TENSOR_DATA_TYPE_UINT16;
 
-		case EMLTensorDataType::UInt8:
+		case ENNETensorDataType::UInt8:
 			return DML_TENSOR_DATA_TYPE_UINT8;
 
-		case EMLTensorDataType::Int64:
+		case ENNETensorDataType::Int64:
 			return DML_TENSOR_DATA_TYPE_INT64;
 
-		case EMLTensorDataType::Int32:
+		case ENNETensorDataType::Int32:
 			return DML_TENSOR_DATA_TYPE_INT32;
 
-		case EMLTensorDataType::Int16:
+		case ENNETensorDataType::Int16:
 			return DML_TENSOR_DATA_TYPE_INT16;
 
-		case EMLTensorDataType::Int8:
+		case ENNETensorDataType::Int8:
 			return DML_TENSOR_DATA_TYPE_INT8;
 
 		default:
@@ -272,7 +272,7 @@ namespace DmlUtil
 		}
 	}
 
-	inline uint64 CalculateBufferSize(const FTensorDesc& DmlTensor, const FTensor& Desc)
+	inline uint64 CalculateBufferSize(const FTensorDesc& DmlTensor, const UE::NNECore::Internal::FTensor& Desc)
 	{
 		uint64 ElemSizeInBytes = Desc.GetElemByteSize();
 		
@@ -320,7 +320,7 @@ public:
 
 	virtual ~FMLOperatorDml() = default;
 
-	virtual bool Initialize(FDeviceContextDml* DevCtx, TArrayView<const FTensor> InputTensors, TArrayView<const FTensor> OutputTensors, const UE::NNECore::FAttributeMap& Attributes) = 0;
+	virtual bool Initialize(FDeviceContextDml* DevCtx, TArrayView<const UE::NNECore::Internal::FTensor> InputTensors, TArrayView<const UE::NNECore::Internal::FTensor> OutputTensors, const UE::NNECore::FAttributeMap& Attributes) = 0;
 
 	IDMLOperator* GetOperator()
 	{
@@ -329,7 +329,7 @@ public:
 
 protected:
 
-	bool InitDmlTensorDesc(DmlUtil::FTensorDesc& DmlTensorDesc, const FTensor& TensorDesc)
+	bool InitDmlTensorDesc(DmlUtil::FTensorDesc& DmlTensorDesc, const UE::NNECore::Internal::FTensor& TensorDesc)
 	{
 		DML_TENSOR_DATA_TYPE DmlDataType = DmlUtil::GetTensorDataType(TensorDesc.GetDataType());
 
@@ -341,7 +341,7 @@ protected:
 			return false;
 		}
 
-		DmlTensorDesc.Sizes = TensorDesc.GetShape().Data;
+		DmlTensorDesc.Sizes = TensorDesc.GetShape().GetData();
 		// TODO: Support tensor padding using strides defined in FMLTensorDesc
 		//DmlUtil::SetTensorStrides(DmlTensorDesc, TensorDesc.Strides);
 		
@@ -360,7 +360,7 @@ protected:
 		return true;
 	}
 
-	bool InitDmlTensorDesc(DmlUtil::FTensorDesc& DmlTensorDesc, const FTensor& TensorDesc, const FTensor& BroadcastDesc)
+	bool InitDmlTensorDesc(DmlUtil::FTensorDesc& DmlTensorDesc, const UE::NNECore::Internal::FTensor& TensorDesc, const UE::NNECore::Internal::FTensor& BroadcastDesc)
 	{
 		DML_TENSOR_DATA_TYPE DmlDataType = DmlUtil::GetTensorDataType(TensorDesc.GetDataType());
 
@@ -374,7 +374,7 @@ protected:
 
 		if (DmlUtil::IsSameShape(TensorDesc, BroadcastDesc))
 		{
-			DmlTensorDesc.Sizes = TensorDesc.GetShape().Data;
+			DmlTensorDesc.Sizes = TensorDesc.GetShape().GetData();
 			DmlUtil::SetTensorStrides(DmlTensorDesc, TensorDesc);
 		}
 		else if (TensorDesc.GetShape().Rank() > BroadcastDesc.GetShape().Rank())
@@ -479,14 +479,14 @@ public:
 	//
 	//
 	//
-	virtual bool Initialize(FDeviceContextDml* InDevCtx, TArrayView<const FTensor> InputTensors, TArrayView<const FTensor> OutputTensors, const UE::NNECore::FAttributeMap& Attributes) override
+	virtual bool Initialize(FDeviceContextDml* InDevCtx, TArrayView<const UE::NNECore::Internal::FTensor> InputTensors, TArrayView<const UE::NNECore::Internal::FTensor> OutputTensors, const UE::NNECore::FAttributeMap& Attributes) override
 	{
 		Num = InputTensors[0].GetVolume();
 
 		DevCtx = InDevCtx;
 
-		const FTensor& InputTensorDesc = InputTensors[0];
-		const FTensor& OutputTensorDesc = OutputTensors[0];
+		const UE::NNECore::Internal::FTensor& InputTensorDesc = InputTensors[0];
+		const UE::NNECore::Internal::FTensor& OutputTensorDesc = OutputTensors[0];
 
 		Alpha = Attributes.GetValueOrDefault(TEXT("alpha"), Alpha);
 		Beta = Attributes.GetValueOrDefault(TEXT("beta"), Beta);
@@ -609,16 +609,16 @@ public:
 	//
 	//
 	//
-	virtual bool Initialize(FDeviceContextDml* InDevCtx, TArrayView<const FTensor> InputTensors, TArrayView<const FTensor> OutputTensors, const UE::NNECore::FAttributeMap& Attributes) override
+	virtual bool Initialize(FDeviceContextDml* InDevCtx, TArrayView<const UE::NNECore::Internal::FTensor> InputTensors, TArrayView<const UE::NNECore::Internal::FTensor> OutputTensors, const UE::NNECore::FAttributeMap& Attributes) override
 	{
 		// TODO: Setup attributes
 		Num = OutputTensors[0].GetVolume();
 
 		DevCtx = InDevCtx;
 
-		const FTensor& InputATensorDesc = InputTensors[0];
-		const FTensor& InputBTensorDesc = InputTensors[1];
-		const FTensor& OutputTensorDesc = OutputTensors[0];
+		const UE::NNECore::Internal::FTensor& InputATensorDesc = InputTensors[0];
+		const UE::NNECore::Internal::FTensor& InputBTensorDesc = InputTensors[1];
+		const UE::NNECore::Internal::FTensor& OutputTensorDesc = OutputTensors[0];
 
 		// Initialize tensor descriptors
 		DmlUtil::FTensorDesc	DmlInputATensorDesc{};
@@ -695,7 +695,7 @@ public:
 	//
 	//
 	//
-	virtual bool Initialize(FDeviceContextDml* InDevCtx, TArrayView<const FTensor> InputTensors, TArrayView<const FTensor> OutputTensors, const UE::NNECore::FAttributeMap& Attributes) override
+	virtual bool Initialize(FDeviceContextDml* InDevCtx, TArrayView<const UE::NNECore::Internal::FTensor> InputTensors, TArrayView<const UE::NNECore::Internal::FTensor> OutputTensors, const UE::NNECore::FAttributeMap& Attributes) override
 	{
 		// Setup attributes
 		float Alpha = 1.0f;
@@ -710,9 +710,9 @@ public:
 		
 		DevCtx = InDevCtx;
 
-		const FTensor& InputATensorDesc = InputTensors[0];
-		const FTensor& InputBTensorDesc = InputTensors[1];
-		const FTensor& OutputTensorDesc = OutputTensors[0];
+		const UE::NNECore::Internal::FTensor& InputATensorDesc = InputTensors[0];
+		const UE::NNECore::Internal::FTensor& InputBTensorDesc = InputTensors[1];
+		const UE::NNECore::Internal::FTensor& OutputTensorDesc = OutputTensors[0];
 
 		// Initialize tensor descriptors
 		DmlUtil::FTensorDesc	DmlInputATensorDesc{};
@@ -734,7 +734,7 @@ public:
 
 		if (InputTensors.Num() > 2)
 		{
-			const FTensor& InputCTensorDesc = InputTensors[2];
+			const UE::NNECore::Internal::FTensor& InputCTensorDesc = InputTensors[2];
 
 			if (!InitDmlTensorDesc(DmlInputCTensorDesc, InputCTensorDesc, OutputTensorDesc))
 			{
@@ -795,7 +795,7 @@ private:
 
 	bool InitCompiledOp(TConstArrayView<int32> OpInputIndices, uint64 TensorDataSize);
 
-	FMLOperatorDml* OpCreate(const FString& Name, TArrayView<const FTensor> InputTensorDesc, TArrayView<const FTensor> OutputTensorDescs, const UE::NNECore::FAttributeMap& Attributes);
+	FMLOperatorDml* OpCreate(const FString& Name, TArrayView<const UE::NNECore::Internal::FTensor> InputTensorDesc, TArrayView<const UE::NNECore::Internal::FTensor> OutputTensorDescs, const UE::NNECore::FAttributeMap& Attributes);
 		
 	FBufferRHIRef CreateRHIBuffer(FRHICommandListImmediate& RHICmdList, uint32 Size, EBufferUsageFlags Usage, ERHIAccess Access, const TCHAR* DbgName);
 	ID3D12Resource* CreateD3D12Buffer(uint32 Size, D3D12_RESOURCE_STATES ResourceState = D3D12_RESOURCE_STATE_COMMON, D3D12_HEAP_TYPE HeapType = D3D12_HEAP_TYPE_DEFAULT, const TCHAR* DebugName = nullptr);
@@ -1370,7 +1370,7 @@ public:
 
 	struct FGraphDesc
 	{
-		TConstArrayView<FTensor>	AllTensors;
+		TConstArrayView<UE::NNECore::Internal::FTensor>	AllTensors;
 		TConstArrayView<int32>		InputIndices;
 		TConstArrayView<int32>		OutputIndices;
 		TConstArrayView<int32>		WeightIndices;
@@ -1812,12 +1812,12 @@ bool FMLInferenceModelDml::Init(TConstArrayView<uint8> ModelData, FDeviceContext
 
 	// DirectML requires all tensors to be concrete
 	// TODO jira 168972: Handle dynamic tensor desc, op should init from symbolic shapes
-	TArray<FTensor>	Tensors;
+	TArray<UE::NNECore::Internal::FTensor>	Tensors;
 
 	Tensors.Reset(AllSymbolicTensorDescs.Num());
-	for (const FTensorDesc& TensorDesc : AllSymbolicTensorDescs)
+	for (const UE::NNECore::FTensorDesc& TensorDesc : AllSymbolicTensorDescs)
 	{
-		Tensors.Emplace(FTensor::MakeFromSymbolicDesc(TensorDesc));
+		Tensors.Emplace(UE::NNECore::Internal::FTensor::MakeFromSymbolicDesc(TensorDesc));
 	}
 
 	FGraphBuilder				DmlGraphBuilder;
@@ -1841,8 +1841,8 @@ bool FMLInferenceModelDml::Init(TConstArrayView<uint8> ModelData, FDeviceContext
 		const FString TypeName = Format.Operators[Idx].TypeName;
 
 		FGraphBuilder::FOpDesc		OpDesc;
-		TArray<FTensor>				OpInputTensors;
-		TArray<FTensor>				OpOutputTensors;
+		TArray<UE::NNECore::Internal::FTensor>				OpInputTensors;
+		TArray<UE::NNECore::Internal::FTensor>				OpOutputTensors;
 		UE::NNECore::FAttributeMap	AttributeMap;
 
 		OpDesc.InputStart = OpInputIndices.Num();
@@ -1862,7 +1862,7 @@ bool FMLInferenceModelDml::Init(TConstArrayView<uint8> ModelData, FDeviceContext
 			{
 				FTensorDesc SymbolicTensorDesc = AllSymbolicTensorDescs[InputTensorIndex];
 				//TODO jira 168972: Handle dynamic tensor desc, op should init from symbolic shapes
-				OpInputTensors.Emplace(FTensor::MakeFromSymbolicDesc(SymbolicTensorDesc));
+				OpInputTensors.Emplace(UE::NNECore::Internal::FTensor::MakeFromSymbolicDesc(SymbolicTensorDesc));
 			}
 
 			OpInputIndices.Emplace(InputTensorIndex);
@@ -1872,7 +1872,7 @@ bool FMLInferenceModelDml::Init(TConstArrayView<uint8> ModelData, FDeviceContext
 		{
 			FTensorDesc SymbolicTensorDesc = AllSymbolicTensorDescs[OutputTensorIndex];
 			//TODO jira 168972: Handle dynamic tensor desc, op should init from symbolic shapes
-			OpOutputTensors.Emplace(FTensor::MakeFromSymbolicDesc(SymbolicTensorDesc));
+			OpOutputTensors.Emplace(UE::NNECore::Internal::FTensor::MakeFromSymbolicDesc(SymbolicTensorDesc));
 			OpOutputIndices.Emplace(OutputTensorIndex);
 		}
 
@@ -2191,7 +2191,7 @@ void FMLInferenceModelDml::AddDispatchOps_RenderThread(FRDGBuilder& GraphBuilder
 //
 // Create operator
 //
-FMLOperatorDml* FMLInferenceModelDml::OpCreate(const FString& OpName, TArrayView<const FTensor> InputTensorDescs, TArrayView<const FTensor> OutputTensorDescs, const UE::NNECore::FAttributeMap& Attributes)
+FMLOperatorDml* FMLInferenceModelDml::OpCreate(const FString& OpName, TArrayView<const UE::NNECore::Internal::FTensor> InputTensorDescs, TArrayView<const UE::NNECore::Internal::FTensor> OutputTensorDescs, const UE::NNECore::FAttributeMap& Attributes)
 {
 	FMLOperatorRegistryDml::OperatorCreateFunc CreateFn = FMLOperatorRegistryDml::Get()->OpFind(OpName);
 
@@ -2268,7 +2268,7 @@ int FMLInferenceModelDml::PrepareTensorShapesAndData()
 {
 	for (FTensorDesc SymbolicTensorDesc : AllSymbolicTensorDescs)
 	{
-		if (!SymbolicTensorDesc.IsConcrete())
+		if (!SymbolicTensorDesc.GetShape().IsConcrete())
 		{
 			UE_LOG(LogNNX, Warning, TEXT("DML engine does not support model with variable shapes yet."));
 			return -1;
diff --git a/Engine/Plugins/Experimental/NNE/Source/NNXRuntimeRDG/Private/NNXRuntimeHLSLHelper.cpp b/Engine/Plugins/Experimental/NNE/Source/NNXRuntimeRDG/Private/NNXRuntimeHLSLHelper.cpp
--- a/Engine/Plugins/Experimental/NNE/Source/NNXRuntimeRDG/Private/NNXRuntimeHLSLHelper.cpp
+++ b/Engine/Plugins/Experimental/NNE/Source/NNXRuntimeRDG/Private/NNXRuntimeHLSLHelper.cpp
@@ -1,21 +1,22 @@
 // Copyright Epic Games, Inc. All Rights Reserved.
 
 #include "NNXRuntimeHLSLHelper.h"
-
+#include "NNECoreTypes.h"
+#include "NNECoreTensor.h"
 #include "RHI.h"
 
 namespace NNX
 {
-	void FillTensorSizeShaderParameters(const FTensor& Tensor, TStaticArray<FUintVector4, NXRT_TENSORSTRIDEINFO_MAX_NUM_DIMENSIONS, 16U>& OutShaderParam, int32 Idx)
+	void FillTensorSizeShaderParameters(const UE::NNECore::Internal::FTensor& Tensor, TStaticArray<FUintVector4, NXRT_TENSORSTRIDEINFO_MAX_NUM_DIMENSIONS, 16U>& OutShaderParam, int32 Idx)
 	{
-		static_assert(FTensorShape::MaxRank <= NXRT_TENSORSTRIDEINFO_MAX_NUM_DIMENSIONS);
+		static_assert(UE::NNECore::FTensorShape::MaxRank <= NXRT_TENSORSTRIDEINFO_MAX_NUM_DIMENSIONS);
 		for (int32 i = 0; i < Tensor.GetShape().Rank(); ++i)
 		{
-			OutShaderParam[i][Idx] = Tensor.GetShape().Data[i];
+			OutShaderParam[i][Idx] = Tensor.GetShape().GetData()[i];
 		}
 	}
 
-	void FillTensorStrideShaderParameters(const FTensor& Tensor, TStaticArray<FUintVector4, NXRT_TENSORSTRIDEINFO_MAX_NUM_DIMENSIONS, 16U>& OutShaderParam, int32 Idx, int32 TargetNumdimensionForBroadcast)
+	void FillTensorStrideShaderParameters(const UE::NNECore::Internal::FTensor& Tensor, TStaticArray<FUintVector4, NXRT_TENSORSTRIDEINFO_MAX_NUM_DIMENSIONS, 16U>& OutShaderParam, int32 Idx, int32 TargetNumdimensionForBroadcast)
 	{
 		if (TargetNumdimensionForBroadcast == -1)
 		{
@@ -24,7 +25,7 @@ namespace NNX
 		checkf(TargetNumdimensionForBroadcast >= Tensor.GetShape().Rank(), TEXT("Can't broadcast tensor from rank %d to rank %d, should be inferior or equal."), Tensor.GetShape().Rank(), TargetNumdimensionForBroadcast);
 		int32 Offset = TargetNumdimensionForBroadcast - Tensor.GetShape().Rank();
 
-		static_assert(FTensorShape::MaxRank <= NXRT_TENSORSTRIDEINFO_MAX_NUM_DIMENSIONS);
+		static_assert(UE::NNECore::FTensorShape::MaxRank <= NXRT_TENSORSTRIDEINFO_MAX_NUM_DIMENSIONS);
 		for (int32 i = 7; i >= 0; --i)
 		{
 			if (i >= TargetNumdimensionForBroadcast || i < Offset)
@@ -37,12 +38,12 @@ namespace NNX
 			}
 			else
 			{
-				OutShaderParam[i][Idx] = OutShaderParam[i + 1][Idx] * Tensor.GetShape().Data[i + 1 - Offset];
+				OutShaderParam[i][Idx] = OutShaderParam[i + 1][Idx] * Tensor.GetShape().GetData()[i + 1 - Offset];
 			}
 		}
 	}
 
-	void FillTensorStrideForBroadcastShaderParameters(const FTensor& Tensor, int32 OutputNumdimension, TStaticArray<FUintVector4, NXRT_TENSORSTRIDEINFO_MAX_NUM_DIMENSIONS, 16U>& OutShaderParam, int32 Idx)
+	void FillTensorStrideForBroadcastShaderParameters(const UE::NNECore::Internal::FTensor& Tensor, int32 OutputNumdimension, TStaticArray<FUintVector4, NXRT_TENSORSTRIDEINFO_MAX_NUM_DIMENSIONS, 16U>& OutShaderParam, int32 Idx)
 	{
 		checkf(OutputNumdimension >= Tensor.GetShape().Rank(), TEXT("Can't broadcast tensor from rank %d to rank %d, should be inferior or equal."), Tensor.GetShape().Rank(), OutputNumdimension);
 		FillTensorStrideShaderParameters(Tensor, OutShaderParam, Idx, OutputNumdimension);
@@ -50,7 +51,7 @@ namespace NNX
 		for (int32 i = Offset; i < OutputNumdimension; ++i)
 		{
 			// the stride for broadcast dimension is kept as 0
-			if (Tensor.GetShape().Data[i - Offset] == 1)
+			if (Tensor.GetShape().GetData()[i - Offset] == 1)
 			{
 				OutShaderParam[i][Idx] = 0;
 			}
diff --git a/Engine/Plugins/Experimental/NNE/Source/NNXRuntimeRDG/Private/NNXRuntimeHLSLHelper.h b/Engine/Plugins/Experimental/NNE/Source/NNXRuntimeRDG/Private/NNXRuntimeHLSLHelper.h
--- a/Engine/Plugins/Experimental/NNE/Source/NNXRuntimeRDG/Private/NNXRuntimeHLSLHelper.h
+++ b/Engine/Plugins/Experimental/NNE/Source/NNXRuntimeRDG/Private/NNXRuntimeHLSLHelper.h
@@ -2,14 +2,15 @@
 
 #pragma once
 
-#include "NNXTypes.h"
 #include "Containers/StaticArray.h"
 #include "NNXShaderParameters.h"
 
+namespace UE::NNECore::Internal { class FTensor; }
+
 namespace NNX
 {
-	void FillTensorSizeShaderParameters(const FTensor& Tensor, TStaticArray<FUintVector4, NXRT_TENSORSTRIDEINFO_MAX_NUM_DIMENSIONS, 16U>& OutShaderParam, int32 Idx);
-	void FillTensorStrideShaderParameters(const FTensor& Tensor, TStaticArray<FUintVector4, NXRT_TENSORSTRIDEINFO_MAX_NUM_DIMENSIONS, 16U>& OutShaderParam, int32 Idx, int32 TargetNumdimensionForBroadcast = -1);
-	void FillTensorStrideForBroadcastShaderParameters(const FTensor& Tensor, int32 OutputNumdimension, TStaticArray<FUintVector4, NXRT_TENSORSTRIDEINFO_MAX_NUM_DIMENSIONS, 16U>& OutShaderParam, int32 Idx);
+	void FillTensorSizeShaderParameters(const UE::NNECore::Internal::FTensor& Tensor, TStaticArray<FUintVector4, NXRT_TENSORSTRIDEINFO_MAX_NUM_DIMENSIONS, 16U>& OutShaderParam, int32 Idx);
+	void FillTensorStrideShaderParameters(const UE::NNECore::Internal::FTensor& Tensor, TStaticArray<FUintVector4, NXRT_TENSORSTRIDEINFO_MAX_NUM_DIMENSIONS, 16U>& OutShaderParam, int32 Idx, int32 TargetNumdimensionForBroadcast = -1);
+	void FillTensorStrideForBroadcastShaderParameters(const UE::NNECore::Internal::FTensor& Tensor, int32 OutputNumdimension, TStaticArray<FUintVector4, NXRT_TENSORSTRIDEINFO_MAX_NUM_DIMENSIONS, 16U>& OutShaderParam, int32 Idx);
 	FIntVector ComputeElementWiseThreadGroups(uint32 ElementCount, uint32 GroupSizeX);
 } // NNX
\ No newline at end of file
diff --git a/Engine/Plugins/Experimental/NNE/Source/NNXRuntimeRDG/Private/NNXRuntimeRDG.cpp b/Engine/Plugins/Experimental/NNE/Source/NNXRuntimeRDG/Private/NNXRuntimeRDG.cpp
--- a/Engine/Plugins/Experimental/NNE/Source/NNXRuntimeRDG/Private/NNXRuntimeRDG.cpp
+++ b/Engine/Plugins/Experimental/NNE/Source/NNXRuntimeRDG/Private/NNXRuntimeRDG.cpp
@@ -30,8 +30,8 @@ struct FMLInferenceModelRDG::FReadbackEntry
 //
 bool AlwaysValidValidationFunction(
 	const UE::NNECore::FAttributeMap& AttributeMap, 
-	TConstArrayView<EMLTensorDataType> InputTensorTypes,
-	TConstArrayView<FSymbolicTensorShape> InputShapes)
+	TConstArrayView<ENNETensorDataType> InputTensorTypes,
+	TConstArrayView<UE::NNECore::FSymbolicTensorShape> InputShapes)
 {
 	return true;
 }
@@ -46,7 +46,7 @@ FInputValidator::FInputValidator() :
 	TemplateTypes.SetNum(1);
 }
 
-bool FInputValidator::Validate(TConstArrayView<EMLTensorDataType> InputTypes)
+bool FInputValidator::Validate(TConstArrayView<ENNETensorDataType> InputTypes)
 {
 	check(InputTemplateIndices.Num() == NumRequiredInput + NumOptionalInput);
 
@@ -81,7 +81,7 @@ void FInputValidator::SetTemplateCount(int TemplateCount)
 {
 	TemplateTypes.SetNum(TemplateCount);
 }
-void FInputValidator::AddSupportedType(EMLTensorDataType Type, int TemplateIdx)
+void FInputValidator::AddSupportedType(ENNETensorDataType Type, int TemplateIdx)
 {
 	check(TemplateTypes.Num() > TemplateIdx);
 	TemplateTypes[TemplateIdx].Add(Type);
@@ -238,8 +238,8 @@ bool FMLInferenceModelRDG::LoadModel(TConstArrayView<uint8> ModelData, FMLRuntim
 	{
 		const FMLFormatTensorDesc& FormatTensorDesc = Format.Tensors[Idx];
 
-		FSymbolicTensorShape SymbolicShape = FSymbolicTensorShape::Make(FormatTensorDesc.Shape);
-		FTensorDesc SymbolicTensor = FTensorDesc::Make(FormatTensorDesc.Name, SymbolicShape, FormatTensorDesc.DataType);
+		const UE::NNECore::FSymbolicTensorShape SymbolicShape = UE::NNECore::FSymbolicTensorShape::Make(FormatTensorDesc.Shape);
+		const UE::NNECore::FTensorDesc SymbolicTensor = UE::NNECore::FTensorDesc::Make(FormatTensorDesc.Name, SymbolicShape, FormatTensorDesc.DataType);
 		
 		AllSymbolicTensorDescs.Emplace(SymbolicTensor);
 		
@@ -260,13 +260,13 @@ bool FMLInferenceModelRDG::LoadModel(TConstArrayView<uint8> ModelData, FMLRuntim
 		else if (FormatTensorDesc.Type == EMLFormatTensorType::Initializer)
 		{
 			WeightTensorIndices.Emplace(Idx);
-			if (!SymbolicTensor.IsConcrete())
+			if (!SymbolicTensor.GetShape().IsConcrete())
 			{
 				UE_LOG(LogNNX, Error, TEXT("Weight tensor %s should have a concrete shape"), *SymbolicTensor.GetName());
 				return false;
 			}
 			
-			const FTensorShape TensorShape = FTensorShape::MakeFromSymbolic(SymbolicTensor.GetShape());
+			const UE::NNECore::FTensorShape TensorShape = UE::NNECore::FTensorShape::MakeFromSymbolic(SymbolicTensor.GetShape());
 			FTensorRDG& WeightRDG = WeightTensorRDGs.Emplace_GetRef(FTensorRDG::Make(SymbolicTensor, TensorShape, nullptr));
 			
 			if (WeightRDG.GetDataSize() != FormatTensorDesc.DataSize)
diff --git a/Engine/Plugins/Experimental/NNE/Source/NNXRuntimeRDG/Private/NNXRuntimeRDG.h b/Engine/Plugins/Experimental/NNE/Source/NNXRuntimeRDG/Private/NNXRuntimeRDG.h
--- a/Engine/Plugins/Experimental/NNE/Source/NNXRuntimeRDG/Private/NNXRuntimeRDG.h
+++ b/Engine/Plugins/Experimental/NNE/Source/NNXRuntimeRDG/Private/NNXRuntimeRDG.h
@@ -4,6 +4,8 @@
 
 #include "NNECoreAttributeMap.h"
 #include "NNECoreAttributeValue.h"
+#include "NNECoreTensor.h"
+#include "NNECoreTypes.h"
 #include "NNXCore.h"
 #include "NNXModelOptimizerInterface.h"
 #include "NNXRuntime.h"
@@ -31,11 +33,12 @@ class FRDGBuilder;
 struct FMLRuntimeFormat;
 struct FNNIModelRaw;
 
+namespace UE::NNECore { class FTensorDesc; }
+
 namespace NNX
 {
 
 struct FTensorBinding;
-class FTensorDesc;
 
 /**
  * Interface for all operators to prepare the model tensors at scheduling time
@@ -43,7 +46,7 @@ class FTensorDesc;
 struct IPrepareOperator
 {
 	virtual ~IPrepareOperator() = default;
-	virtual int PrepareOutputs(TConstArrayView<NNX::FTensorRef> InputTensors, TArrayView<NNX::FTensorRef> OutputTensors) const = 0;
+	virtual int PrepareOutputs(TConstArrayView<UE::NNECore::Internal::FTensorRef> InputTensors, TArrayView<UE::NNECore::Internal::FTensorRef> OutputTensors) const = 0;
 };
 
 /**
@@ -54,7 +57,7 @@ struct IOperatorRDG
 	virtual ~IOperatorRDG() = default;
 };
 
-class FTensorRDG : public FTensor
+class FTensorRDG : public UE::NNECore::Internal::FTensor
 {
 	FRDGBufferRef Buffer{};
 	
@@ -68,7 +71,7 @@ public:
 		TensorRDG.DataType = TensorDesc.GetDataType();
 		TensorRDG.Shape = Shape;
 		TensorRDG.Volume = Shape.Volume();
-		TensorRDG.DataSize = (uint64)GetTensorDataTypeSizeInBytes(TensorRDG.DataType) * TensorRDG.Volume;
+		TensorRDG.DataSize = (uint64)UE::NNECore::GetTensorDataTypeSizeInBytes(TensorRDG.DataType) * TensorRDG.Volume;
 		TensorRDG.Buffer = Buffer;
 		check(TensorRDG.Volume <= TNumericLimits<uint32>::Max());
 		return TensorRDG;
@@ -150,7 +153,7 @@ protected:
 };
 
 //TODO jira 167584 remove default validation and declare contract in all DML operator (see HLSL Gemm for current example)
-bool AlwaysValidValidationFunction(const UE::NNECore::FAttributeMap& AttributeMap, TConstArrayView<EMLTensorDataType> InputTypes, TConstArrayView<FSymbolicTensorShape> InputShapes);
+bool AlwaysValidValidationFunction(const UE::NNECore::FAttributeMap& AttributeMap, TConstArrayView<ENNETensorDataType> InputTypes, TConstArrayView<UE::NNECore::FSymbolicTensorShape> InputShapes);
 
 class FInputValidator
 {
@@ -159,11 +162,11 @@ public:
 	void AddRequired(int32 TemplateIdx = 0);
 	void AddOptional(int32 TemplateIdx = 0);
 	void SetTemplateCount(int32 TemplateCount);
-	void AddSupportedType(EMLTensorDataType Type, int32 TemplateIdx = 0);
-	bool Validate(TConstArrayView<EMLTensorDataType> InputTypes);
+	void AddSupportedType(ENNETensorDataType Type, int32 TemplateIdx = 0);
+	bool Validate(TConstArrayView<ENNETensorDataType> InputTypes);
 
 private:
-	TArray<TArray<EMLTensorDataType>> TemplateTypes;
+	TArray<TArray<ENNETensorDataType>> TemplateTypes;
 	TArray<int32> InputTemplateIndices;
 	int32 NumRequiredInput;
 	int32 NumOptionalInput;
@@ -205,7 +208,7 @@ class TOperatorRegistryRDG
 public:
 
 	typedef TOperatorType* (*OperatorCreateFunc)();
-	typedef bool (*OperatorValidateFunc)(const UE::NNECore::FAttributeMap& AttributeMap, TConstArrayView<EMLTensorDataType> InputTypes, TConstArrayView<FSymbolicTensorShape> InputShapes);
+	typedef bool (*OperatorValidateFunc)(const UE::NNECore::FAttributeMap& AttributeMap, TConstArrayView<ENNETensorDataType> InputTypes, TConstArrayView<UE::NNECore::FSymbolicTensorShape> InputShapes);
 
 	static TOperatorRegistryRDG* Get()
 	{
@@ -290,14 +293,14 @@ public:
 
 		for (int32 Idx = 0; Idx < Format.Operators.Num(); ++Idx)
 		{
-			TArray<EMLTensorDataType> InputTensorTypes;
-			TArray<FSymbolicTensorShape> InputTensorShapes;
+			TArray<ENNETensorDataType> InputTensorTypes;
+			TArray<UE::NNECore::FSymbolicTensorShape> InputTensorShapes;
 			UE::NNECore::FAttributeMap AttributeMap;
 			
 			for (int32 InputTensorIndex: Format.Operators[Idx].InTensors)
 			{
 				InputTensorTypes.Add(Format.Tensors[InputTensorIndex].DataType);
-				InputTensorShapes.Add(FSymbolicTensorShape::Make(Format.Tensors[InputTensorIndex].Shape));
+				InputTensorShapes.Add(UE::NNECore::FSymbolicTensorShape::Make(Format.Tensors[InputTensorIndex].Shape));
 			}
 			for (const FMLFormatAttributeDesc& Desc : Format.Operators[Idx].Attributes)
 			{
diff --git a/Engine/Plugins/Experimental/NNE/Source/NNXUtils/Private/NNXModelBuilderNNX.cpp b/Engine/Plugins/Experimental/NNE/Source/NNXUtils/Private/NNXModelBuilderNNX.cpp
--- a/Engine/Plugins/Experimental/NNE/Source/NNXUtils/Private/NNXModelBuilderNNX.cpp
+++ b/Engine/Plugins/Experimental/NNE/Source/NNXUtils/Private/NNXModelBuilderNNX.cpp
@@ -109,7 +109,7 @@ public:
 		return !Data.IsEmpty();
 	}
 
-	virtual HTensor AddTensor(const FString& Name, EMLTensorDataType DataType, TArrayView<const int32> Shape, const void* Data, uint64 DataSize)
+	virtual HTensor AddTensor(const FString& Name, ENNETensorDataType DataType, TArrayView<const int32> Shape, const void* Data, uint64 DataSize)
 	{
 		TArray<int32, TInlineAllocator<FTensorShape::MaxRank>> NNIShape;
 		for (int i = 0; i < Shape.Num(); ++i)
@@ -239,7 +239,7 @@ public:
 
 private:
 
-	int AddTensor(const FString& InName, TArrayView<const int32> InShape, EMLTensorDataType InDataType, const void* Data, uint64 DataSize)
+	int AddTensor(const FString& InName, TArrayView<const int32> InShape, ENNETensorDataType InDataType, const void* Data, uint64 DataSize)
 	{
 		int Idx = -1;
 
diff --git a/Engine/Plugins/Experimental/NNE/Source/NNXUtils/Private/NNXModelBuilderONNX.cpp b/Engine/Plugins/Experimental/NNE/Source/NNXUtils/Private/NNXModelBuilderONNX.cpp
--- a/Engine/Plugins/Experimental/NNE/Source/NNXUtils/Private/NNXModelBuilderONNX.cpp
+++ b/Engine/Plugins/Experimental/NNE/Source/NNXUtils/Private/NNXModelBuilderONNX.cpp
@@ -112,7 +112,7 @@ public:
 	}
 
 	/** Add tensor */
-	virtual HTensor AddTensor(const FString& Name, EMLTensorDataType DataType, TArrayView<const int32> Shape, const void* Data, uint64 DataSize) override
+	virtual HTensor AddTensor(const FString& Name, ENNETensorDataType DataType, TArrayView<const int32> Shape, const void* Data, uint64 DataSize) override
 	{
 		auto Value = onnx::ValueInfoProto().New(Graph->GetArena());
 		
@@ -132,7 +132,7 @@ public:
 				Tensor->add_dims(Shape[Idx]);
 			}
 
-			checkf(DataType != EMLTensorDataType::Char, TEXT("Char tensors not supported yet!"));
+			checkf(DataType != ENNETensorDataType::Char, TEXT("Char tensors not supported yet!"));
 			// raw_data is not supported for strings, when implementing those one will need to use string_data from TensorProto
 			std::string* raw_data = Tensor->mutable_raw_data();
 			raw_data->assign((const char*)Data, DataSize);
@@ -260,7 +260,7 @@ public:
 
 private:
 
-	bool SetValue(onnx::ValueInfoProto* Value, const FString& Name, EMLTensorDataType DataType, const TArrayView<const int32>& InShape)
+	bool SetValue(onnx::ValueInfoProto* Value, const FString& Name, ENNETensorDataType DataType, const TArrayView<const int32>& InShape)
 	{
 		onnx::TypeProto*		Type = Value->mutable_type();
 		onnx::TypeProto_Tensor* TensorType = Type->mutable_tensor_type();
@@ -279,27 +279,27 @@ private:
 		return true;
 	}
 
-	onnx::TensorProto_DataType ToTensorProtoDataType(EMLTensorDataType DataType)
+	onnx::TensorProto_DataType ToTensorProtoDataType(ENNETensorDataType DataType)
 	{
 		switch (DataType)
 		{
-			case EMLTensorDataType::None: return onnx::TensorProto_DataType_UNDEFINED;
-			case EMLTensorDataType::Float: return onnx::TensorProto_DataType_FLOAT;
-			case EMLTensorDataType::UInt8: return onnx::TensorProto_DataType_UINT8;
-			case EMLTensorDataType::Int8: return onnx::TensorProto_DataType_INT8;
-			case EMLTensorDataType::UInt16: return onnx::TensorProto_DataType_UINT16;
-			case EMLTensorDataType::Int16: return onnx::TensorProto_DataType_INT16;
-			case EMLTensorDataType::Int32: return onnx::TensorProto_DataType_INT32;
-			case EMLTensorDataType::Int64: return onnx::TensorProto_DataType_INT64;
-			//case EMLTensorDataType::String: return onnx::TensorProto_DataType_STRING;
-			case EMLTensorDataType::Boolean: return onnx::TensorProto_DataType_BOOL;
-			case EMLTensorDataType::Half: return onnx::TensorProto_DataType_FLOAT16;
-			case EMLTensorDataType::Double: return onnx::TensorProto_DataType_DOUBLE;
-			case EMLTensorDataType::UInt32: return onnx::TensorProto_DataType_UINT32;
-			case EMLTensorDataType::UInt64: return onnx::TensorProto_DataType_UINT64;
-			case EMLTensorDataType::Complex64: return onnx::TensorProto_DataType_COMPLEX64;
-			case EMLTensorDataType::Complex128: return onnx::TensorProto_DataType_COMPLEX128;
-			case EMLTensorDataType::BFloat16: return onnx::TensorProto_DataType_BFLOAT16;
+			case ENNETensorDataType::None: return onnx::TensorProto_DataType_UNDEFINED;
+			case ENNETensorDataType::Float: return onnx::TensorProto_DataType_FLOAT;
+			case ENNETensorDataType::UInt8: return onnx::TensorProto_DataType_UINT8;
+			case ENNETensorDataType::Int8: return onnx::TensorProto_DataType_INT8;
+			case ENNETensorDataType::UInt16: return onnx::TensorProto_DataType_UINT16;
+			case ENNETensorDataType::Int16: return onnx::TensorProto_DataType_INT16;
+			case ENNETensorDataType::Int32: return onnx::TensorProto_DataType_INT32;
+			case ENNETensorDataType::Int64: return onnx::TensorProto_DataType_INT64;
+			//case ENNETensorDataType::String: return onnx::TensorProto_DataType_STRING;
+			case ENNETensorDataType::Boolean: return onnx::TensorProto_DataType_BOOL;
+			case ENNETensorDataType::Half: return onnx::TensorProto_DataType_FLOAT16;
+			case ENNETensorDataType::Double: return onnx::TensorProto_DataType_DOUBLE;
+			case ENNETensorDataType::UInt32: return onnx::TensorProto_DataType_UINT32;
+			case ENNETensorDataType::UInt64: return onnx::TensorProto_DataType_UINT64;
+			case ENNETensorDataType::Complex64: return onnx::TensorProto_DataType_COMPLEX64;
+			case ENNETensorDataType::Complex128: return onnx::TensorProto_DataType_COMPLEX128;
+			case ENNETensorDataType::BFloat16: return onnx::TensorProto_DataType_BFLOAT16;
 			default: return onnx::TensorProto_DataType_UNDEFINED;
 		}
 	}
@@ -313,7 +313,7 @@ void BuildShapeForModel(bool ConvertToVariadicShape, const FTensorShape& InShape
 	OutShape.Empty();
 	for (int32 Idx = 0; Idx < InShape.Rank(); ++Idx)
 	{
-		int32 Dim = static_cast<int32>(InShape.Data[Idx]);
+		int32 Dim = static_cast<int32>(InShape.GetData()[Idx]);
 		if (ConvertToVariadicShape)
 		{
 			Dim = -1;
diff --git a/Engine/Plugins/Experimental/NNE/Source/NNXUtils/Private/NNXModelOptimizer.cpp b/Engine/Plugins/Experimental/NNE/Source/NNXUtils/Private/NNXModelOptimizer.cpp
--- a/Engine/Plugins/Experimental/NNE/Source/NNXUtils/Private/NNXModelOptimizer.cpp
+++ b/Engine/Plugins/Experimental/NNE/Source/NNXUtils/Private/NNXModelOptimizer.cpp
@@ -231,9 +231,9 @@ private:
 #undef Print
 
 
-#define CASE(GType, Type) case Ort::GraphTensorDataType::GType: return EMLTensorDataType::Type
+#define CASE(GType, Type) case Ort::GraphTensorDataType::GType: return ENNETensorDataType::Type
 
-EMLTensorDataType GetDataTypeFromGraphTensor(Ort::GraphTensorDataType TensorDataType)
+ENNETensorDataType GetDataTypeFromGraphTensor(Ort::GraphTensorDataType TensorDataType)
 {
 	switch (TensorDataType)
 	{
@@ -255,7 +255,7 @@ EMLTensorDataType GetDataTypeFromGraphTensor(Ort::GraphTensorDataType TensorData
 		CASE(kBFloat16, BFloat16);
 
 		default:
-			return EMLTensorDataType::None;
+			return ENNETensorDataType::None;
 	}
 }
 
@@ -500,7 +500,7 @@ private:
 		{
 			Ort::GraphTensorInfo TensorInfo = Graph->GetGraphInput(Idx);
 
-			EMLTensorDataType DataType = GetDataTypeFromGraphTensor(TensorInfo.dataType);
+			ENNETensorDataType DataType = GetDataTypeFromGraphTensor(TensorInfo.dataType);
 
 			auto Tensor = 
 				Builder->AddTensor(
@@ -517,7 +517,7 @@ private:
 		{
 			Ort::GraphTensorInfo TensorInfo = Graph->GetGraphOutput(Idx);
 
-			EMLTensorDataType DataType = GetDataTypeFromGraphTensor(TensorInfo.dataType);
+			ENNETensorDataType DataType = GetDataTypeFromGraphTensor(TensorInfo.dataType);
 
 			auto Tensor =
 				Builder->AddTensor(
@@ -579,7 +579,7 @@ private:
 			{
 				Ort::GraphTensorInfo		TensorInfo = Graph->GetNodeInput(Node, InIdx);
 				Ort::GraphTensorInitializer	TensorInit = Graph->GetTensorInitializer(TensorInfo.name);
-				EMLTensorDataType			DataType = GetDataTypeFromGraphTensor(TensorInfo.dataType);
+				ENNETensorDataType			DataType = GetDataTypeFromGraphTensor(TensorInfo.dataType);
 				const void*					Data = nullptr;
 				uint64						DataSize = 0;
 
@@ -613,7 +613,7 @@ private:
 			for (int OutIdx = 0; OutIdx < NodeInfo.outputCount; ++OutIdx)
 			{
 				Ort::GraphTensorInfo		TensorInfo = Graph->GetNodeOutput(Node, OutIdx);
-				EMLTensorDataType			DataType = GetDataTypeFromGraphTensor(TensorInfo.dataType);
+				ENNETensorDataType			DataType = GetDataTypeFromGraphTensor(TensorInfo.dataType);
 				
 				auto Tensor =
 					Builder->AddTensor(
diff --git a/Engine/Plugins/Experimental/NNE/Source/NNXUtils/Public/NNXModelBuilder.h b/Engine/Plugins/Experimental/NNE/Source/NNXUtils/Public/NNXModelBuilder.h
--- a/Engine/Plugins/Experimental/NNE/Source/NNXUtils/Public/NNXModelBuilder.h
+++ b/Engine/Plugins/Experimental/NNE/Source/NNXUtils/Public/NNXModelBuilder.h
@@ -1,7 +1,8 @@
 // Copyright Epic Games, Inc. All Rights Reserved.
 #pragma once
 
-#include "NNXTypes.h"
+#include "NNECoreTypes.h"
+#include "NNECoreTensor.h"
 
 namespace UE::NNECore { class FAttributeMap; }
 struct FNNEAttributeValue;
@@ -9,6 +10,7 @@ struct FNNIModelRaw;
 
 namespace NNX
 {
+using FTensor = UE::NNECore::Internal::FTensor;
 
 class IMLModelBuilder
 {
@@ -52,7 +54,7 @@ public:
 	virtual bool End(TArray<uint8>& Data) = 0;
 
 	/** Add tensor */
-	virtual HTensor AddTensor(const FString& Name, EMLTensorDataType DataType, TArrayView<const int32> Shape, const void* Data = nullptr, uint64 DataSize = 0) = 0;
+	virtual HTensor AddTensor(const FString& Name, ENNETensorDataType DataType, TArrayView<const int32> Shape, const void* Data = nullptr, uint64 DataSize = 0) = 0;
 
 	/** Add model input */
 	virtual bool AddInput(HTensor InTensor) = 0;
