NNE - Set IsBetaVersion to false as it is still experimental

diff --git a/Engine/Plugins/Experimental/NNE/NNE.uplugin b/Engine/Plugins/Experimental/NNE/NNE.uplugin
--- a/Engine/Plugins/Experimental/NNE/NNE.uplugin
+++ b/Engine/Plugins/Experimental/NNE/NNE.uplugin
@@ -1,7 +1,7 @@
 {
 	"FileVersion": 3,
-	"Version": 0.5,
-	"VersionName": "0.5.0-alpha",
+	"Version": 0.1,
+	"VersionName": "0.1",
 	"FriendlyName": "NNE",
 	"Description": "Cross-platform framework for running deep learning and neural network inference in Unreal Engine, supporting both CPU and GPU acceleration.",
 	"Category": "Experimental",
@@ -12,7 +12,7 @@
 	"SupportURL": "",
 	"EnabledByDefault": false,
 	"CanContainContent": false,
-	"IsBetaVersion": true,
+	"IsBetaVersion": false,
 	"IsExperimentalVersion": true,
 	"Installed": false,
 	"SupportedTargetPlatforms": [
