NNE - Implement numpy-like isclose comparison for tests.
https://numpy.org/doc/stable/reference/generated/numpy.isclose.html
- Switch from absolute & relative errors separate checks to combined inequality
- Change nomenclature to reflect numpy's meaning. Relative error -> relative tolerance, absolute epsilon -> absolute tolerance.
- Switch away from using percentage for relative tolerance, so to be consistent with the convention adopted by numpy and most numerical codes.
- Change tolerance defaults to numpy ones (stricter than before).
- Adapt tolerances for specific tests in NNXQATestDesc.json: removed all relative tolerance overrides (not needed), kept only absolute tolerances when needed.

diff --git a/Engine/Plugins/Experimental/NNE/Source/NNXQA/Private/NNXQAJsonUtils.h b/Engine/Plugins/Experimental/NNE/Source/NNXQA/Private/NNXQAJsonUtils.h
--- a/Engine/Plugins/Experimental/NNE/Source/NNXQA/Private/NNXQAJsonUtils.h
+++ b/Engine/Plugins/Experimental/NNE/Source/NNXQA/Private/NNXQAJsonUtils.h
@@ -14,6 +14,8 @@ namespace Test
 {
 namespace Json
 {
+	constexpr float JSON_TOLERANCE_NOTSET = -1.0f; 
+
 	struct FTestConfigTensor : FJsonSerializable
 	{
 		TArray<int32> Shape;
@@ -42,14 +44,14 @@ namespace Json
 	{
 		FString Name;
 		bool Skip;
-		float AbsoluteError;
-		float RelativeError;
+		float AbsoluteTolerance = JSON_TOLERANCE_NOTSET;
+		float RelativeTolerance = JSON_TOLERANCE_NOTSET;
 
 		BEGIN_JSON_SERIALIZER
 			JSON_SERIALIZE("name", Name);
 			JSON_SERIALIZE("skip", Skip);
-			JSON_SERIALIZE("absolute_error", AbsoluteError);
-			JSON_SERIALIZE("relative_error", RelativeError);
+			JSON_SERIALIZE("absolute_tolerance", AbsoluteTolerance);
+			JSON_SERIALIZE("relative_tolerance", RelativeTolerance);
 		END_JSON_SERIALIZER
 	};
 
@@ -139,8 +141,8 @@ namespace Json
 		TArray<FString> AdditionalDatasets;
 		TArray<FString> RemovedDatasets;
 		bool Skip;
-		float AbsoluteError;
-		float RelativeError;
+		float AbsoluteTolerance = JSON_TOLERANCE_NOTSET;
+		float RelativeTolerance = JSON_TOLERANCE_NOTSET;
 		FString InputType;
 		FString OutputType;
 		TArray<FTestConfigRuntime> Runtimes;
@@ -148,8 +150,8 @@ namespace Json
 		BEGIN_JSON_SERIALIZER
 			JSON_SERIALIZE("target", Target);
 			JSON_SERIALIZE("skip", Skip);
-			JSON_SERIALIZE("absolute_error", AbsoluteError);
-			JSON_SERIALIZE("relative_error", RelativeError);
+			JSON_SERIALIZE("absolute_tolerance", AbsoluteTolerance);
+			JSON_SERIALIZE("relative_tolerance", RelativeTolerance);
 			JSON_SERIALIZE("input_type", InputType);
 			JSON_SERIALIZE("output_type", OutputType);
 			JSON_SERIALIZE_ARRAY_SERIALIZABLE("runtimes", Runtimes, FTestConfigRuntime);
diff --git a/Engine/Plugins/Experimental/NNE/Source/NNXQA/Private/NNXQAParametricTest.cpp b/Engine/Plugins/Experimental/NNE/Source/NNXQA/Private/NNXQAParametricTest.cpp
--- a/Engine/Plugins/Experimental/NNE/Source/NNXQA/Private/NNXQAParametricTest.cpp
+++ b/Engine/Plugins/Experimental/NNE/Source/NNXQA/Private/NNXQAParametricTest.cpp
@@ -70,14 +70,14 @@ namespace Test
 		
 		static void ApplyEpsilons(FTests::FTestSetup& TestSetup, const Json::FTestConfigTarget& TestTarget)
 		{
-			if (TestTarget.AbsoluteError != 0.0f)
+			if (TestTarget.AbsoluteTolerance != Json::JSON_TOLERANCE_NOTSET)
 			{
-				TestSetup.AbsoluteErrorEpsilon = TestTarget.AbsoluteError;
+				TestSetup.AbsoluteTolerance = TestTarget.AbsoluteTolerance;
 			}
 			
-			if (TestTarget.RelativeError != 0.0f)
+			if (TestTarget.RelativeTolerance != Json::JSON_TOLERANCE_NOTSET)
 			{
-				TestSetup.RelativeErrorPercent = TestTarget.RelativeError;
+				TestSetup.RelativeTolerance = TestTarget.RelativeTolerance;
 			}
 		}
 
@@ -91,13 +91,13 @@ namespace Test
 				}
 				else
 				{
-					if (Runtime.AbsoluteError != 0.0f)
+					if (Runtime.AbsoluteTolerance != Json::JSON_TOLERANCE_NOTSET)
 					{
-						TestSetup.AbsoluteErrorEpsilonForRuntime.Emplace(Runtime.Name, Runtime.AbsoluteError);
+						TestSetup.AbsoluteToleranceForRuntime.Emplace(Runtime.Name, Runtime.AbsoluteTolerance);
 					}
-					if (Runtime.RelativeError != 0.0f)
+					if (Runtime.RelativeTolerance != Json::JSON_TOLERANCE_NOTSET)
 					{
-						TestSetup.RelativeErrorPercentForRuntime.Emplace(Runtime.Name, Runtime.RelativeError);
+						TestSetup.RelativeToleranceForRuntime.Emplace(Runtime.Name, Runtime.RelativeTolerance);
 					}
 				}
 			}
@@ -106,14 +106,7 @@ namespace Test
 		static void ApplyTargetConfig(FTests::FTestSetup& TestSetup, const Json::FTestConfigTarget& TestTarget)
 		{
 			//Epsilons
-			if (TestTarget.AbsoluteError != 0.0f)
-			{
-				TestSetup.AbsoluteErrorEpsilon = TestTarget.AbsoluteError;
-			}
-			if (TestTarget.RelativeError != 0.0f)
-			{
-				TestSetup.RelativeErrorPercent = TestTarget.RelativeError;
-			}
+			ApplyEpsilons(TestSetup, TestTarget);
 
 			ApplyRuntimesConfig(TestSetup, TestTarget.Runtimes);
 
diff --git a/Engine/Plugins/Experimental/NNE/Source/NNXQA/Private/NNXQAUtils.cpp b/Engine/Plugins/Experimental/NNE/Source/NNXQA/Private/NNXQAUtils.cpp
--- a/Engine/Plugins/Experimental/NNE/Source/NNXQA/Private/NNXQAUtils.cpp
+++ b/Engine/Plugins/Experimental/NNE/Source/NNXQA/Private/NNXQAUtils.cpp
@@ -220,8 +220,9 @@ namespace Test
 	template<typename T> bool CompareTensorData(
 		const UE::NNECore::Internal::FTensor& RefTensorDesc,   const TArray<char>& RefRawBuffer,
 		const UE::NNECore::Internal::FTensor& OtherTensorDesc, const TArray<char>& OtherRawBuffer,
-		float AbsoluteErrorEpsilon, float RelativeErrorPercent)
+		float AbsoluteTolerance, float RelativeTolerance)
 	{
+
 		const T* RefBuffer = (T*)RefRawBuffer.GetData();
 		const T* OtherBuffer = (T*)OtherRawBuffer.GetData();
 		
@@ -279,37 +280,47 @@ namespace Test
 			{
 				continue;
 			}
-
 			const float AbsoluteError = FMath::Abs<float>(Result - Reference);
-			const float RelativeError = 100.0f * (AbsoluteError / FMath::Abs<float>(Reference));
-			if (AbsoluteError > AbsoluteErrorEpsilon || RelativeError > RelativeErrorPercent)
+			const float AbsoluteRef = FMath::Abs<float>(Reference);
+			const float AbsoluteThreshold = AbsoluteTolerance + RelativeTolerance * AbsoluteRef;
+			/* from https://numpy.org/doc/stable/reference/generated/numpy.isclose.html */
+			if(AbsoluteError > AbsoluteThreshold)
 			{
 				bTensorMemMatch = false;
-				if (AbsoluteError > WorstAbsoluteError)
+				// Depending on dominating term, contribute to different statistics
+				if(AbsoluteError > RelativeTolerance * AbsoluteRef)
+				{
+					WorstRelativeError = AbsoluteError / AbsoluteRef;
+					WorstRelativeErrorIndex = i;
+					WorstRelativeErrorRef = Reference;
+					WorstRelativeErrorOther = Result;
+				}
+				else
 				{
 					WorstAbsoluteError = AbsoluteError;
 					WorstAbsoluteErrorIndex = i;
 					WorstAbsoluteErrorRef = Reference;
 					WorstAbsoluteErrorOther = Result;
 				}
-				if (RelativeError > WorstRelativeError)
-				{
-					WorstRelativeError = RelativeError;
-					WorstRelativeErrorIndex = i;
-					WorstRelativeErrorRef = Reference;
-					WorstRelativeErrorOther = Result;
-				}
 			}
 		}
 
 		if (!bTensorMemMatch)
 		{
+			FString AbsoluteErrorMessage = WorstAbsoluteErrorIndex == -1 ? FString() : FString::Printf(
+				TEXT("LogNNX: Worst absolute tolerance violation %.3e (tol %.3e) at position %d, got %.3e expected %.3e\n"),
+				WorstAbsoluteError, AbsoluteTolerance, WorstAbsoluteErrorIndex, WorstAbsoluteErrorOther, WorstAbsoluteErrorRef
+			);
+			FString RelativeErrorMessage = WorstRelativeErrorIndex == -1 ? FString() : FString::Printf(
+				TEXT("LogNNX: Worst relative tolerance violation %.3e (tol %.3e) at position %d, got %.3e expected %.3e\n"), 
+				WorstRelativeError, RelativeTolerance, WorstRelativeErrorIndex, WorstRelativeErrorOther, WorstRelativeErrorRef
+			); 
 			UE_LOG(LogNNX, Error, TEXT("Tensor data do not match.\n"
-				"LogNNX: Worst absolute error %f (epsilon %f) at position %d, got %f expected %f\n"
-				"LogNNX: Worst relative error % f % %(epsilon % f%%) at position % d, got % f expected % f\n"
+				"%s"
+				"%s"
 				"LogNNX: Num unexpected NaNs %d (First at index %d), Num missing NaNs %d (first at index %d)"),
-				WorstAbsoluteError, AbsoluteErrorEpsilon, WorstAbsoluteErrorIndex, WorstAbsoluteErrorOther, WorstAbsoluteErrorRef,
-				WorstRelativeError, RelativeErrorPercent, WorstRelativeErrorIndex, WorstRelativeErrorOther, WorstRelativeErrorRef,
+				*AbsoluteErrorMessage,
+				*RelativeErrorMessage,
 				NumExtraNaNsInResults, FirstExtraNaNIndex, NumMissingNaNsInResults, FirstMissingNaNIndex);
 			UE_LOG(LogNNX, Error, TEXT("   Expected : %s"), *TensorToString(RefTensorDesc, RefRawBuffer));
 			UE_LOG(LogNNX, Error, TEXT("   But got  : %s"), *TensorToString(OtherTensorDesc, OtherRawBuffer));
@@ -322,7 +333,7 @@ namespace Test
 	bool VerifyTensorResult(
 		const UE::NNECore::Internal::FTensor& RefTensor, const TArray<char>& RefRawBuffer,
 		const UE::NNECore::Internal::FTensor& OtherTensor, const TArray<char>& OtherRawBuffer,
-		float AbsoluteErrorEpsilon, float RelativeErrorPercent)
+		float AbsoluteTolerance, float RelativeTolerance)
 	{
 		if (RefTensor.GetShape() != OtherTensor.GetShape())
 		{
@@ -334,20 +345,20 @@ namespace Test
 
 		if (RefTensor.GetDataType() == ENNETensorDataType::Float)
 		{
-			return CompareTensorData<float>(RefTensor, RefRawBuffer, OtherTensor, OtherRawBuffer, AbsoluteErrorEpsilon, RelativeErrorPercent);
+			return CompareTensorData<float>(RefTensor, RefRawBuffer, OtherTensor, OtherRawBuffer, AbsoluteTolerance, RelativeTolerance);
 		}
 		else if (RefTensor.GetDataType() == ENNETensorDataType::Boolean)
 		{
 			check(RefTensor.GetElemByteSize() == 1);
-			return CompareTensorData<bool>(RefTensor, RefRawBuffer, OtherTensor, OtherRawBuffer, AbsoluteErrorEpsilon, RelativeErrorPercent);
+			return CompareTensorData<bool>(RefTensor, RefRawBuffer, OtherTensor, OtherRawBuffer, AbsoluteTolerance, RelativeTolerance);
 		}
 		else if (RefTensor.GetDataType() == ENNETensorDataType::Int32)
 		{
-			return CompareTensorData<int32>(RefTensor, RefRawBuffer, OtherTensor, OtherRawBuffer, AbsoluteErrorEpsilon, RelativeErrorPercent);
+			return CompareTensorData<int32>(RefTensor, RefRawBuffer, OtherTensor, OtherRawBuffer, AbsoluteTolerance, RelativeTolerance);
 		}
 		else if (RefTensor.GetDataType() == ENNETensorDataType::UInt32)
 		{
-			return CompareTensorData<uint32>(RefTensor, RefRawBuffer, OtherTensor, OtherRawBuffer, AbsoluteErrorEpsilon, RelativeErrorPercent);
+			return CompareTensorData<uint32>(RefTensor, RefRawBuffer, OtherTensor, OtherRawBuffer, AbsoluteTolerance, RelativeTolerance);
 		}
 		else
 		{
@@ -452,8 +463,8 @@ namespace Test
 		TArray<UE::NNECore::Internal::FTensor> OutputTensors;
 
 		const FString& RuntimeName = Runtime->GetRuntimeName();
-		const float AbsoluteErrorEpsilon = TestSetup.GetAbsoluteErrorEpsilonForRuntime(RuntimeName);
-		const float RelativeErrorPercent = TestSetup.GetRelativeErrorPercentForRuntime(RuntimeName);
+		const float AbsoluteTolerance = TestSetup.GetAbsoluteToleranceForRuntime(RuntimeName);
+		const float RelativeTolerance = TestSetup.GetRelativeToleranceForRuntime(RuntimeName);
 		bool bTestSuceeded = true;
 	
 		if (!RunTestInference(ONNXModel, TestSetup, Runtime, OutputTensors, OutputMemBuffers))
@@ -469,7 +480,7 @@ namespace Test
 				bTestSuceeded &= VerifyTensorResult(
 					RefOutputTensors[i], RefOutputMemBuffers[i],
 					OutputTensors[i], OutputMemBuffers[i],
-					AbsoluteErrorEpsilon, RelativeErrorPercent);
+					AbsoluteTolerance, RelativeTolerance);
 			}
 		}
 		else
@@ -498,8 +509,8 @@ namespace Test
 			return false;
 		}
 		const FString& RefName = RefRuntime->GetRuntimeName();
-		const float AbsoluteRefErrorEpsilon = TestSetup.GetAbsoluteErrorEpsilonForRuntime(RefName);
-		const float RelativeRefErrorPercent = TestSetup.GetRelativeErrorPercentForRuntime(RefName);
+		const float AbsoluteTolerance = TestSetup.GetAbsoluteToleranceForRuntime(RefName);
+		const float RelativeTolerance = TestSetup.GetRelativeToleranceForRuntime(RefName);
 		TArray<TArray<char>> RefOutputMemBuffers;
 		TArray<UE::NNECore::Internal::FTensor> RefOutputTensors;
 		bool bAllTestsSucceeded = true;
@@ -518,7 +529,7 @@ namespace Test
 				bAllTestsSucceeded &= VerifyTensorResult(
 					TestSetup.Outputs[i], TestSetup.OutputsData[i],
 					RefOutputTensors[i], RefOutputMemBuffers[i],
-					AbsoluteRefErrorEpsilon, RelativeRefErrorPercent);
+					AbsoluteTolerance, RelativeTolerance);
 			}
 		}
 		if (!bAllTestsSucceeded)
diff --git a/Engine/Plugins/Experimental/NNE/Source/NNXQA/Private/NNXQAUtils.h b/Engine/Plugins/Experimental/NNE/Source/NNXQA/Private/NNXQAUtils.h
--- a/Engine/Plugins/Experimental/NNE/Source/NNXQA/Private/NNXQAUtils.h
+++ b/Engine/Plugins/Experimental/NNE/Source/NNXQA/Private/NNXQAUtils.h
@@ -22,24 +22,25 @@ namespace Test
 		
 		struct FTestSetup
 		{
-			static constexpr float DefaultAbsoluteErrorEpsilon = 1e-5f;
-			static constexpr float DefaultRelativeErrorPercent = 1e-3f;
+			/* defaults from https://numpy.org/doc/stable/reference/generated/numpy.isclose.html */
+			static constexpr float DefaultAbsoluteTolerance = 1e-8f;
+			static constexpr float DefaultRelativeTolerance = 1e-5f;
 
 			FTestSetup(const FString& InTestCategory, const FString& InModelOrOperatorName, const FString& InTestSuffix) :
 				TestName(InTestCategory + InModelOrOperatorName + InTestSuffix),
 				TargetName(InModelOrOperatorName),
-				AbsoluteErrorEpsilon(DefaultAbsoluteErrorEpsilon),
-				RelativeErrorPercent(DefaultRelativeErrorPercent),
+				AbsoluteTolerance(DefaultAbsoluteTolerance),
+				RelativeTolerance(DefaultRelativeTolerance),
 				IsModelTest(false)
 			{}
 			
 			const FString TestName;
 			const FString TargetName;
-			float AbsoluteErrorEpsilon;
-			float RelativeErrorPercent;
+			float AbsoluteTolerance;
+			float RelativeTolerance;
 			bool IsModelTest;
-			TMap<FString, float> AbsoluteErrorEpsilonForRuntime;
-			TMap<FString, float> RelativeErrorPercentForRuntime;
+			TMap<FString, float> AbsoluteToleranceForRuntime;
+			TMap<FString, float> RelativeToleranceForRuntime;
 			TArray<UE::NNECore::Internal::FTensor> Inputs;
 			TArray<UE::NNECore::Internal::FTensor> Weights;
 			TArray<UE::NNECore::Internal::FTensor> Outputs;
@@ -52,16 +53,16 @@ namespace Test
 			TArray<FString> AutomationExcludedPlatform;
 			TArray<TPair<FString, FString>> AutomationExcludedPlatformRuntimeCombination;
 
-			float GetAbsoluteErrorEpsilonForRuntime(const FString& RuntimeName) const
+			float GetAbsoluteToleranceForRuntime(const FString& RuntimeName) const
 			{
-				const float* SpecializedValue = AbsoluteErrorEpsilonForRuntime.Find(RuntimeName);
-				return (SpecializedValue != nullptr) ? *SpecializedValue : AbsoluteErrorEpsilon;
+				const float* SpecializedValue = AbsoluteToleranceForRuntime.Find(RuntimeName);
+				return (SpecializedValue != nullptr) ? *SpecializedValue : AbsoluteTolerance;
 
 			}
-			float GetRelativeErrorPercentForRuntime(const FString& RuntimeName) const
+			float GetRelativeToleranceForRuntime(const FString& RuntimeName) const
 			{
-				const float* SpecializedValue = RelativeErrorPercentForRuntime.Find(RuntimeName);
-				return (SpecializedValue != nullptr) ? *SpecializedValue : RelativeErrorPercent;
+				const float* SpecializedValue = RelativeToleranceForRuntime.Find(RuntimeName);
+				return (SpecializedValue != nullptr) ? *SpecializedValue : RelativeTolerance;
 			}
 		};
 
diff --git a/Engine/Plugins/Experimental/NNE/Source/NNXQA/Resources/NNXQATestDesc.json b/Engine/Plugins/Experimental/NNE/Source/NNXQA/Resources/NNXQATestDesc.json
--- a/Engine/Plugins/Experimental/NNE/Source/NNXQA/Resources/NNXQATestDesc.json
+++ b/Engine/Plugins/Experimental/NNE/Source/NNXQA/Resources/NNXQATestDesc.json
@@ -9,22 +9,22 @@
 					"target": "resnet50-v2-7",
 					"tags": [ "SmokeTest" ],
 					"runtimes": [
-						{ "name": "NNXRuntimeORTDml", "relative_error": 1.0 },
+						{ "name": "NNXRuntimeORTDml" },
 						{ "name": "NNXRuntimeHlsl",  "skip": true, "comment": "Missing BatchNormalisation operator." },
 						{ "name": "NNXRuntimeDml",  "skip": true, "comment": "Model not supported yet." }
-					]
+					],
+					"absolute_tolerance": 1e-5
 				},
 				{
 					"target": "Mosaic-9",
 					"tags": [ "SmokeTest" ],
 					"runtimes": [
-						{ "name": "NNXRuntimeORTDml", "absolute_error": 5e-2, "relative_error": 50.0,
-							"comment": "Worst relative error 43.367836 % (epsilon 0.001000 %) at position 112558, got -0.003391 expected -0.002365. TODO we need an automatic balance between relative error and absolute error"
-						},
-						{ "name": "NNXRuntimeORTCpu", "absolute_error": 1e-3, "relative_error": 3.0 },
-						{ "name": "NNXRuntimeHlsl", "absolute_error": 0.03, "relative_error": 55.0, "comment": "TODO check, errors are very high." },
+						{ "name": "NNXRuntimeORTDml" },
+						{ "name": "NNXRuntimeORTCpu" },
+						{ "name": "NNXRuntimeHlsl" },
 						{ "name": "NNXRuntimeDml",  "skip": true, "comment": "Model not supported yet." }
-					]
+					],
+					"absolute_tolerance": 5.0e-2
 				}
 			]
 		},
@@ -36,10 +36,11 @@
 					"target": "NeuralMorphModel_global",
 					"tags": [ "SmokeTest" ],
 					"runtimes": [
-						{ "name": "NNXRuntimeORTDml", "relative_error": 0.025 },
-						{ "name": "NNXRuntimeHlsl", "relative_error": 0.025 },
-						{ "name": "NNXRuntimeDml", "relative_error": 0.025 }
-					]
+						{ "name": "NNXRuntimeORTDml" },
+						{ "name": "NNXRuntimeHlsl" },
+						{ "name": "NNXRuntimeDml" }
+					],
+					"absolute_tolerance": 1e-5
 				},
 				{
 					"target": "NeuralMorphModel_global_Dax",
@@ -54,7 +55,7 @@
 					"target": "NeuralMorphModel_Local",
 					"tags": [ "SmokeTest" ],
 					"runtimes": [
-						{ "name": "NNXRuntimeORTDml", "relative_error": 0.04 },
+						{ "name": "NNXRuntimeORTDml", "absolute_tolerance": 1e-6 },
 						{ "name": "NNXRuntimeHlsl",  "skip": true, "comment": "Missing Reshape operator." },
 						{ "name": "NNXRuntimeDml",  "skip": true, "comment": "Model not supported yet." }
 					]
@@ -62,8 +63,8 @@
 				{
 					"target": "VertexDeltaModel",
 					"runtimes": [
-						{ "name": "NNXRuntimeORTDml", "absolute_error": 1e-4, "relative_error": 10.0, "comment": "TODO check, relative error is very high."	},
-						{ "name": "NNXRuntimeHlsl", "absolute_error": 1e-4, "relative_error": 3.0 },
+						{ "name": "NNXRuntimeORTDml", "absolute_tolerance": 1e-4 },
+						{ "name": "NNXRuntimeHlsl", "absolute_tolerance": 1e-4 },
 						{ "name": "NNXRuntimeDml",  "skip": true, "comment": "Model not supported yet." }
 					]
 				}
@@ -124,25 +125,25 @@
 			"targets": [
 				{	"target": "Abs", "additional_datasets": [ "Elementwise.Slow.OneInput" ]},
 				{	"target": "Acos", "tags": [ "SmokeTest" ], "runtimes": [
-						{ "name": "NNXRuntimeORTDml", "absolute_error":  1e-4, "relative_error": 0.03},
-						{ "name": "NNXRuntimeDml",    "absolute_error":  1e-4, "relative_error": 0.03},
-						{ "name": "NNXRuntimeHlsl",   "absolute_error":  2e-4, "relative_error": 0.03}
+						{ "name": "NNXRuntimeORTDml", "absolute_tolerance":  1e-4},
+						{ "name": "NNXRuntimeDml",    "absolute_tolerance":  1e-4},
+						{ "name": "NNXRuntimeHlsl",   "absolute_tolerance":  2e-4}
 					]},
 				{	"target": "Acosh" },
 				{	"target": "Asin", "runtimes": [
-						{ "name": "NNXRuntimeORTDml", "absolute_error":  1e-4, "relative_error": 2.5},
-						{ "name": "NNXRuntimeDml",    "absolute_error":  1e-4, "relative_error": 2.5},
-						{ "name": "NNXRuntimeHlsl",   "absolute_error":  2e-4, "relative_error": 2.5}
+						{ "name": "NNXRuntimeORTDml", "absolute_tolerance":  1e-4},
+						{ "name": "NNXRuntimeDml",    "absolute_tolerance":  1e-4},
+						{ "name": "NNXRuntimeHlsl",   "absolute_tolerance":  2e-4}
 					]},
 				{	"target": "Asinh", "runtimes": [
-						{ "name": "NNXRuntimeORTDml", "absolute_error":  1e-4, "relative_error": 0.03},
-						{ "name": "NNXRuntimeDml",    "absolute_error":  1e-4, "relative_error": 0.03},
-						{ "name": "NNXRuntimeHlsl",   "absolute_error":  1e-4, "relative_error": 0.03}
+						{ "name": "NNXRuntimeORTDml", "absolute_tolerance":  1e-4},
+						{ "name": "NNXRuntimeDml",    "absolute_tolerance":  1e-4},
+						{ "name": "NNXRuntimeHlsl",   "absolute_tolerance":  2e-4}
 					]},
 				{	"target": "Atan", "runtimes": [
-						{ "name": "NNXRuntimeORTDml", "absolute_error":  1e-4, "relative_error": 0.03},
-						{ "name": "NNXRuntimeDml",    "absolute_error":  1e-4, "relative_error": 0.03},
-						{ "name": "NNXRuntimeHlsl",   "absolute_error":  1.5e-3, "relative_error": 0.3}
+						{ "name": "NNXRuntimeORTDml", "absolute_tolerance":  1e-4},
+						{ "name": "NNXRuntimeDml",    "absolute_tolerance":  1e-4},
+						{ "name": "NNXRuntimeHlsl",   "absolute_tolerance":  1.5e-3}
 					]},
 				{	"target": "Atanh" },
 				{	"target": "BitShift", "skip": true, "comment": "TODO need attributes" },
@@ -155,14 +156,14 @@
 				{	"target": "Cos" },
 				{	"target": "Cosh" },
 				{	"target": "Elu", "comment": "TODO test with not-default attribute", "runtimes": [
-						{ "name": "NNXRuntimeORTDml", "absolute_error":  1e-4, "relative_error": 0.03},
-						{ "name": "NNXRuntimeDml",    "absolute_error":  1e-4, "relative_error": 0.03},
-						{ "name": "NNXRuntimeHlsl",   "absolute_error":  1e-4, "relative_error": 0.03}
+						{ "name": "NNXRuntimeORTDml", "absolute_tolerance":  1e-4},
+						{ "name": "NNXRuntimeDml",    "absolute_tolerance":  1e-4},
+						{ "name": "NNXRuntimeHlsl",   "absolute_tolerance":  1e-4}
 					]},
 				{	"target": "Erf", "tags": [ "SmokeTest" ], "runtimes": [
-						{ "name": "NNXRuntimeORTDml", "absolute_error":  1e-4, "relative_error": 0.03},
-						{ "name": "NNXRuntimeDml",    "absolute_error":  1e-4, "relative_error": 0.03},
-						{ "name": "NNXRuntimeHlsl",   "absolute_error":  5e-4, "relative_error": 0.05, "comment": "NNXRuntimeHLSL Erf implementation is less precise than the DML one"}
+						{ "name": "NNXRuntimeORTDml", "absolute_tolerance":  1e-4},
+						{ "name": "NNXRuntimeDml",    "absolute_tolerance":  1e-4},
+						{ "name": "NNXRuntimeHlsl",   "absolute_tolerance":  5e-4, "comment": "NNXRuntimeHLSL Erf implementation is less precise than the DML one"}
 					]},
 				{	"target": "Exp" },
 				{	"target": "Floor" },
@@ -172,10 +173,10 @@
 				{	"target": "HardSwish", "runtimes": [ { "name": "NNXRuntimeDml", "skip": true, "comment": "Not implemented atm on NNXRuntimeDML."} ]},
 				{	"target": "LeakyRelu", "additional_datasets": [ "Elementwise.LeakyRelu.AttributeTest" ] },
 				{	"target": "Log", "runtimes": [
-						{ "name": "NNXRuntimeORTDml", "relative_error": 0.15},
-						{ "name": "NNXRuntimeDml",    "relative_error": 0.15},
-						{ "name": "NNXRuntimeHlsl",   "relative_error": 0.15}
-					]},
+						{ "name": "NNXRuntimeORTDml"},
+						{ "name": "NNXRuntimeDml"},
+						{ "name": "NNXRuntimeHlsl"}
+					], "absolute_tolerance": 1e-6 },
 				{	"target": "Neg" },
 				{	"target": "Not", "input_type": "boolean", "output_type": "boolean", "runtimes": [
 						{ "name": "NNXRuntimeDml", "skip": true, "comment": "RDG need bool tensor support"},
@@ -185,30 +186,30 @@
 				{	"target": "Relu", "tags": [ "SmokeTest" ] },
 				{	"target": "Round", "tags": [ "SmokeTest" ] },
 				{	"target": "Selu", "comment": "TODO test with not-default attribute", "runtimes": [
-						{ "name": "NNXRuntimeORTDml", "absolute_error":  1e-4, "relative_error": 0.03},
-						{ "name": "NNXRuntimeDml",    "absolute_error":  1e-4, "relative_error": 0.03},
-						{ "name": "NNXRuntimeHlsl",   "absolute_error":  1e-4, "relative_error": 0.03}
+						{ "name": "NNXRuntimeORTDml", "absolute_tolerance":  1e-4},
+						{ "name": "NNXRuntimeDml",    "absolute_tolerance":  1e-4},
+						{ "name": "NNXRuntimeHlsl",   "absolute_tolerance":  1e-4}
 					]},
 				{	"target": "Sigmoid", "tags": [ "SmokeTest" ] },
 				{	"target": "Sign" },
 				{	"target": "Sin", "runtimes": [
-						{ "name": "NNXRuntimeORTDml", "absolute_error":  1e-4, "relative_error": 0.03},
-						{ "name": "NNXRuntimeDml",    "absolute_error":  1e-4, "relative_error": 0.03},
-						{ "name": "NNXRuntimeHlsl",   "absolute_error":  1e-4, "relative_error": 0.03}
+						{ "name": "NNXRuntimeORTDml", "absolute_tolerance":  1e-4},
+						{ "name": "NNXRuntimeDml",    "absolute_tolerance":  1e-4},
+						{ "name": "NNXRuntimeHlsl",   "absolute_tolerance":  1e-4}
 					]},
 				{	"target": "Sinh" },
 				{	"target": "Softplus", "tags": [ "SmokeTest" ] },
 				{	"target": "Softsign", "tags": [ "SmokeTest" ] },
 				{	"target": "Sqrt", "tags": [ "SmokeTest" ] },
 				{	"target": "Tan", "runtimes": [
-						{ "name": "NNXRuntimeORTDml", "absolute_error":  1e-4, "relative_error": 0.03},
-						{ "name": "NNXRuntimeDml",    "absolute_error":  1e-4, "relative_error": 0.03},
-						{ "name": "NNXRuntimeHlsl",   "absolute_error":  1e-4, "relative_error": 0.03}
+						{ "name": "NNXRuntimeORTDml", "absolute_tolerance":  1e-4},
+						{ "name": "NNXRuntimeDml",    "absolute_tolerance":  1e-4},
+						{ "name": "NNXRuntimeHlsl",   "absolute_tolerance":  1e-4}
 					]},
 				{	"target": "Tanh", "runtimes": [
-						{ "name": "NNXRuntimeORTDml", "absolute_error":  1e-4, "relative_error": 0.03},
-						{ "name": "NNXRuntimeDml",    "absolute_error":  1e-4, "relative_error": 0.03},
-						{ "name": "NNXRuntimeHlsl",   "absolute_error":  1e-4, "relative_error": 0.03}
+						{ "name": "NNXRuntimeORTDml", "absolute_tolerance":  1e-4},
+						{ "name": "NNXRuntimeDml",    "absolute_tolerance":  1e-4},
+						{ "name": "NNXRuntimeHlsl",   "absolute_tolerance":  1e-4}
 					]}
 			]
 		},
@@ -251,7 +252,7 @@
 						"Elementwise.Binary.BidirectionalBroadcast", 
 						"Elementwise.Binary.LHSToRHSBroadcast" ], 
 					"comment": "Prelu does not support bidirectional broadcast. It only support slope(RHS) to input0(LHS)"},
-				{	"target": "Pow", "absolute_error": 1e-4},
+				{	"target": "Pow", "absolute_tolerance": 1e-4},
 				{	"target": "Sub"},
 				{	"target": "Xor", "input_type": "boolean", "output_type": "boolean", "runtimes": [
 						{ "name": "NNXRuntimeDml", "skip": true, "comment": "RDG need bool tensor support"},
@@ -280,22 +281,22 @@
 			"category": "Convolution",
 			"runtimes": [
 				{ "name": "NNXRuntimeDml", "skip": true	},
-				{ "name": "NNXRuntimeORTDml", "relative_error": 3.5e-2 },
-				{ "name": "NNXRuntimeHlsl",   "relative_error": 3.5e-2 }
+				{ "name": "NNXRuntimeORTDml" },
+				{ "name": "NNXRuntimeHlsl" }
 			],
 			"targets": [
-				{ "target": "Conv" }
+				{ "target": "Conv", "absolute_tolerance": 1e-6 }
 			]
 		},
 		{
 			"category": "ConvTranspose",
 			"runtimes": [
 				{ "name": "NNXRuntimeDml", "skip": true	},
-				{ "name": "NNXRuntimeORTDml", "relative_error": 3.5e-2 },
-				{ "name": "NNXRuntimeHlsl",   "relative_error": 3.5e-2 }
+				{ "name": "NNXRuntimeORTDml" },
+				{ "name": "NNXRuntimeHlsl" }
 			],
 			"targets": [
-				{ "target": "ConvTranspose" }
+				{ "target": "ConvTranspose", "absolute_tolerance": 1e-6 }
 			]
 		},
 		{
@@ -320,7 +321,7 @@
 				{ "name": "NNXRuntimeDml", "skip": true }
 			],
 			"targets": [
-				{ "target": "InstanceNormalization" }
+				{ "target": "InstanceNormalization", "absolute_tolerance": 1e-5 }
 			]
 		},
 		{
@@ -586,7 +587,7 @@
 			"name": "MatMul.Rank5",
 			"runtimes": [	
 				{ "name": "NNXRuntimeORTDml", "skip": true, "comment": "DML metacommand only support 2 to 4 for inputs rank" },
-				{ "name": "NNXRuntimeHlsl", "relative_error": 0.01}
+				{ "name": "NNXRuntimeHlsl"}
 			],
 			"datasets": [
 				{ "inputs": [ { "shape": [ 2, 2, 3, 5, 4 ] }, { "shape": [ 2, 1, 3, 4, 3 ] } ], "outputs": [ { "shape": [ 2, 2, 3, 5, 3 ] } ]},
@@ -663,65 +664,65 @@
 				{
 					"inputs": [ { "shape": [ 2, 2, 4, 1, 3 ] }, { "shape": [ 2 ] }, { "shape": [ 2 ] } ], "outputs": [ { "shape": [ 2, 2, 4, 1, 3 ] } ],
 					"runtimes": [
-						{ "name": "NNXRuntimeHLSL", "relative_error": 3.5e-3 },
-						{ "name": "NNXRuntimeORTCpu", "relative_error": 3.5e-3 },
+						{ "name": "NNXRuntimeHLSL" },
+						{ "name": "NNXRuntimeORTCpu" },
 						{ "name": "NNXRuntimeORTDml", "skip": true, "comment": "Not supported" }
 					]
 				},
 				{
 					"inputs": [ { "shape": [ 2, 2, 4, 4 ] }, { "shape": [ 2 ] }, { "shape": [ 2 ] } ], "outputs": [ { "shape": [ 2, 2, 4, 4 ] } ],
 					"runtimes": [
-						{ "name": "NNXRuntimeHLSL", "relative_error": 3.5e-3 },
-						{ "name": "NNXRuntimeORTCpu", "relative_error": 3.5e-3 },
-						{ "name": "NNXRuntimeORTDml", "relative_error": 3.5e-3 }
+						{ "name": "NNXRuntimeHLSL" },
+						{ "name": "NNXRuntimeORTCpu" },
+						{ "name": "NNXRuntimeORTDml" }
 					]
 				},
 				{
 					"inputs": [ { "shape": [ 2, 2, 8, 8 ] }, { "shape": [ 2 ] }, { "shape": [ 2 ] } ], "outputs": [ { "shape": [ 2, 2, 8, 8 ] } ],
 					"runtimes": [
-						{ "name": "NNXRuntimeHLSL", "relative_error": 2 },
-						{ "name": "NNXRuntimeORTCpu", "relative_error": 2 },
-						{ "name": "NNXRuntimeORTDml", "relative_error": 2 }
+						{ "name": "NNXRuntimeHLSL" },
+						{ "name": "NNXRuntimeORTCpu" },
+						{ "name": "NNXRuntimeORTDml" }
 					]
 				},
 				{
 					"inputs": [ { "shape": [ 2, 128, 8, 8 ] }, { "shape": [ 128 ] }, { "shape": [ 128 ] } ], "outputs": [ { "shape": [ 2, 128, 8, 8 ] } ],
 					"runtimes": [
-						{ "name": "NNXRuntimeHLSL", "relative_error": 2 },
-						{ "name": "NNXRuntimeORTCpu", "relative_error": 2 },
-						{ "name": "NNXRuntimeORTDml", "relative_error": 2 }
+						{ "name": "NNXRuntimeHLSL" },
+						{ "name": "NNXRuntimeORTCpu" },
+						{ "name": "NNXRuntimeORTDml" }
 					]
 				},
 				{
 					"inputs": [ { "shape": [ 128, 2, 8, 8 ] }, { "shape": [ 2 ] }, { "shape": [ 2 ] } ], "outputs": [ { "shape": [ 128, 2, 8, 8 ] } ],
 					"runtimes": [
-						{ "name": "NNXRuntimeHLSL", "relative_error": 2 },
-						{ "name": "NNXRuntimeORTCpu", "relative_error": 2 },
-						{ "name": "NNXRuntimeORTDml", "relative_error": 2 }
+						{ "name": "NNXRuntimeHLSL" },
+						{ "name": "NNXRuntimeORTCpu" },
+						{ "name": "NNXRuntimeORTDml" }
 					]
 				},
 				{
 					"inputs": [ { "shape": [ 64, 64, 8, 8 ] }, { "shape": [ 64 ] }, { "shape": [ 64 ] } ], "outputs": [ { "shape": [ 64, 64, 8, 8 ] } ],
 					"runtimes": [
-						{ "name": "NNXRuntimeHLSL", "relative_error": 2 },
-						{ "name": "NNXRuntimeORTCpu", "relative_error": 2 },
-						{ "name": "NNXRuntimeORTDml", "relative_error": 2 }
+						{ "name": "NNXRuntimeHLSL" },
+						{ "name": "NNXRuntimeORTCpu" },
+						{ "name": "NNXRuntimeORTDml" }
 					]
 				},
 				{
 					"inputs": [ { "shape": [ 2, 2, 45, 45 ] }, { "shape": [ 2 ] }, { "shape": [ 2 ] } ], "outputs": [ { "shape": [ 2, 2, 45, 45 ] } ],
 					"runtimes": [
-						{ "name": "NNXRuntimeHLSL", "relative_error": 5 },
-						{ "name": "NNXRuntimeORTCpu", "relative_error": 5 },
-						{ "name": "NNXRuntimeORTDml", "relative_error": 5 }
+						{ "name": "NNXRuntimeHLSL" },
+						{ "name": "NNXRuntimeORTCpu" },
+						{ "name": "NNXRuntimeORTDml" }
 					]
 				},
 				{
 					"inputs": [ { "shape": [ 1, 128, 56, 56 ] }, { "shape": [ 128 ] }, { "shape": [ 128 ] } ], "outputs": [ { "shape": [ 1, 128, 56, 56 ] } ],
 					"runtimes": [
-						{ "name": "NNXRuntimeHLSL", "relative_error": 10.5 },
-						{ "name": "NNXRuntimeORTCpu", "relative_error": 9 },
-						{ "name": "NNXRuntimeORTDml", "relative_error": 9 }
+						{ "name": "NNXRuntimeHLSL" },
+						{ "name": "NNXRuntimeORTCpu" },
+						{ "name": "NNXRuntimeORTDml" }
 					],
 					"comment": "from 'mosaic-9-sim.onnx'"
 				}
diff --git a/Engine/Plugins/Experimental/NNE/Source/NNXQA/Resources/NNXQATestDescSchema.json b/Engine/Plugins/Experimental/NNE/Source/NNXQA/Resources/NNXQATestDescSchema.json
--- a/Engine/Plugins/Experimental/NNE/Source/NNXQA/Resources/NNXQATestDescSchema.json
+++ b/Engine/Plugins/Experimental/NNE/Source/NNXQA/Resources/NNXQATestDescSchema.json
@@ -65,16 +65,16 @@
 					"type": "boolean",
 					"description": "true to skip the test for this runtime, false by default"
 				},
-				"absolute_error": {
+				"absolute_tolerance": {
 					"type": "number",
 					"minimum": 0,
-					"description": "max absolute error, 0 to keep engine default (1e-5)"
+					"description": "absolute tolerance of numerical tensor comparison routine"
 				},
-				"relative_error": {
+				"relative_tolerance": {
 					"type": "number",
 					"minimum": 0,
-					"maximum": 100,
-					"description": "max relative error in percent, 0 to keep engine default (1e-4 in %)"
+					"maximum": 1,
+					"description": "relative tolerance of numerical tensor comparison routine"
 				},
 				"comment": {
 					"type": "string",
@@ -176,16 +176,16 @@
 					"type": "boolean",
 					"description": "true to skip this operator or model test entirely, false by default"
 				},
-				"absolute_error": {
+				"absolute_tolerance": {
 					"type": "number",
 					"minimum": 0,
-					"description": "max absolute error, 0 to keep engine default (1e-5)"
+					"description": "absolute tolerance (equality check)"
 				},
-				"relative_error": {
+				"relative_tolerance": {
 					"type": "number",
 					"minimum": 0,
-					"maximum": 100,
-					"description": "max relative error in percent, 0 to keep engine default (1e-4 in %)"
+					"maximum": 1,
+					"description": "relative tolerance (equality check)"
 				},
 				"input_type": {
 					"type": "string",
