NNE : Minimal API - NNECore with NNERuntimeCPU as a proof of concept.

Problem:
* NNE API surface is large atm potentially creating a support burden and harming adoption for 5.2.

Solution:
* Reduce the API surface as discussed in design discutions.
* This 2nd CL is targeting NNECore + NNERuntimeCPU (using a 2nd registry to keep both API until all runtime are ported).
* NNERuntimeCPU (minimal API) is used as reference for tests in NNXQA.

Test:
* All tests green under win64 editor

diff --git a/Engine/Plugins/Experimental/NNE/Source/NNXCore/Internal/NNECoreModelBase.h b/Engine/Plugins/Experimental/NNE/Source/NNXCore/Internal/NNECoreModelBase.h
new file mode 100644
--- /dev/null
+++ b/Engine/Plugins/Experimental/NNE/Source/NNXCore/Internal/NNECoreModelBase.h
@@ -0,0 +1,80 @@
+// Copyright Epic Games, Inc. All Rights Reserved.
+
+#pragma once
+
+#include "CoreMinimal.h"
+#include "NNECore.h"
+#include "NNECoreTypes.h"
+
+namespace UE::NNECore::Internal
+{
+	template <class ModelInterface> class FModelBase : public ModelInterface
+	{
+	public:
+
+		virtual ~FModelBase() = default;
+
+		virtual TConstArrayView<NNECore::FTensorDesc> GetInputTensorDescs() const override;
+		virtual TConstArrayView<NNECore::FTensorDesc> GetOutputTensorDescs() const override;
+		virtual TConstArrayView<NNECore::FTensorShape> GetInputTensorShapes() const override;
+		virtual TConstArrayView<NNECore::FTensorShape> GetOutputTensorShapes() const override;
+		virtual int SetInputTensorShapes(TConstArrayView<NNECore::FTensorShape> InInputShapes) override;
+
+	protected:
+
+		FModelBase() {}
+
+		TArray<FTensorShape>	InputTensorShapes;
+		TArray<FTensorShape>	OutputTensorShapes;
+		TArray<FTensorDesc>		InputSymbolicTensors;
+		TArray<FTensorDesc>		OutputSymbolicTensors;
+	};
+
+	template <class T> TConstArrayView<FTensorDesc> FModelBase<T>::GetInputTensorDescs() const
+	{
+		return InputSymbolicTensors;
+	}
+
+	template <class T> TConstArrayView<FTensorDesc> FModelBase<T>::GetOutputTensorDescs() const
+	{
+		return OutputSymbolicTensors;
+	}
+
+	template <class T> TConstArrayView<FTensorShape> FModelBase<T>::GetInputTensorShapes() const
+	{
+		return InputTensorShapes;
+	}
+
+	template <class T> TConstArrayView<FTensorShape> FModelBase<T>::GetOutputTensorShapes() const
+	{
+		return OutputTensorShapes;
+	}
+
+	template <class T> int FModelBase<T>::SetInputTensorShapes(TConstArrayView<FTensorShape> InInputShapes)
+	{
+		InputTensorShapes.Empty();
+
+		if (InInputShapes.Num() != InputSymbolicTensors.Num())
+		{
+			UE_LOG(LogNNE, Warning, TEXT("Number of input shapes does not match number of input tensors"));
+			return -1;
+		}
+
+		for (int32 i = 0; i < InInputShapes.Num(); ++i)
+		{
+			const FTensorDesc SymbolicDesc = InputSymbolicTensors[i];
+			if (!InInputShapes[i].IsCompatibleWith(SymbolicDesc.GetShape()))
+			{
+				UE_LOG(LogNNE, Warning, TEXT("Input shape does not match input tensor %s of index %d"), *SymbolicDesc.GetName(), i);
+				return -1;
+			}
+		}
+
+		InputTensorShapes = InInputShapes;
+
+		//Implementations are responsible to handle output and intermediate tensor shape inference.
+		//This base implementation only validate that all inputs are matching what the model can support.
+		return 0;
+	}
+
+} // namespace UE::NNECore::Internal
\ No newline at end of file
diff --git a/Engine/Plugins/Experimental/NNE/Source/NNXCore/Private/NNECore.cpp b/Engine/Plugins/Experimental/NNE/Source/NNXCore/Private/NNECore.cpp
new file mode 100644
--- /dev/null
+++ b/Engine/Plugins/Experimental/NNE/Source/NNXCore/Private/NNECore.cpp
@@ -0,0 +1,104 @@
+// Copyright Epic Games, Inc. All Rights Reserved.
+
+#include "NNECore.h"
+
+DEFINE_LOG_CATEGORY(LogNNE);
+
+namespace UE::NNECore
+{
+	class FRegistry
+	{
+	public:
+
+		static FRegistry* Get()
+		{
+			static FRegistry Inst;
+
+			return &Inst;
+		}
+
+		bool Add(TWeakInterfacePtr<INNERuntime> Runtime)
+		{
+			if (!Runtime.IsValid())
+			{
+				return false;
+			}
+			
+			if (FindByName(Runtime->GetRuntimeName()) >= 0)
+			{
+				UE_LOG(LogNNE, Warning, TEXT("Runtime %s is already registered"), *Runtime->GetRuntimeName());
+				return false;
+			}
+
+			Runtimes.Add(Runtime);
+
+			return true;
+		}
+
+		bool Remove(TWeakInterfacePtr<INNERuntime> Runtime)
+		{
+			if (!Runtime.IsValid())
+			{
+				return false;
+			}
+			
+			int Index = FindByName(Runtime->GetRuntimeName());
+			if (Index >= 0)
+			{
+				Runtimes.RemoveAtSwap(Index);
+				return true;
+			}
+
+			return false;
+		}
+
+
+		TArrayView<TWeakInterfacePtr<INNERuntime>> GetAllRuntimes()
+		{
+			RemoveInvalidPtrs();
+			return Runtimes;
+		}
+
+	private:
+
+		void RemoveInvalidPtrs()
+		{
+			Runtimes.RemoveAllSwap([](TWeakInterfacePtr<INNERuntime> Runtime) -> bool {return !Runtime.IsValid();}, false);
+		}
+
+		int FindByName(const FString& Name)
+		{
+			RemoveInvalidPtrs();
+			
+			int Index = -1;
+
+			for (int Idx = 0; Idx < Runtimes.Num(); ++Idx)
+			{
+				if (Runtimes[Idx]->GetRuntimeName() == Name)
+				{
+					Index = Idx;
+					break;
+				}
+			}
+
+			return Index;
+		}
+
+		TArray<TWeakInterfacePtr<INNERuntime>>	Runtimes;
+	};
+
+	bool RegisterRuntime(TWeakInterfacePtr<INNERuntime> Runtime)
+	{
+		return FRegistry::Get()->Add(Runtime);
+	}
+
+	bool UnregisterRuntime(TWeakInterfacePtr<INNERuntime> Runtime)
+	{
+		return FRegistry::Get()->Remove(Runtime);
+	}
+
+	TArrayView<TWeakInterfacePtr<INNERuntime>> GetAllRuntimes()
+	{
+		return FRegistry::Get()->GetAllRuntimes();
+	}
+}
diff --git a/Engine/Plugins/Experimental/NNE/Source/NNXCore/Public/NNECore.h b/Engine/Plugins/Experimental/NNE/Source/NNXCore/Public/NNECore.h
new file mode 100644
--- /dev/null
+++ b/Engine/Plugins/Experimental/NNE/Source/NNXCore/Public/NNECore.h
@@ -0,0 +1,35 @@
+// Copyright Epic Games, Inc. All Rights Reserved.
+
+#pragma once
+
+#include "CoreMinimal.h"
+#include "NNECoreRuntime.h"
+#include "UObject/WeakInterfacePtr.h"
+
+NNXCORE_API DECLARE_LOG_CATEGORY_EXTERN(LogNNE, Log, All);
+
+namespace UE::NNECore
+{
+	NNXCORE_API bool RegisterRuntime(TWeakInterfacePtr<INNERuntime> Runtime);
+
+	NNXCORE_API bool UnregisterRuntime(TWeakInterfacePtr<INNERuntime> Runtime);
+	
+	// Enumerate all available runtime modules
+	NNXCORE_API TArrayView<TWeakInterfacePtr<INNERuntime>> GetAllRuntimes();
+
+	/// Get runtime by name
+	template<class T> TWeakInterfacePtr<T> GetRuntime(const FString& Name)
+	{
+		for (TWeakInterfacePtr<INNERuntime> Runtime : GetAllRuntimes())
+		{
+			if (Runtime->GetRuntimeName() == Name)
+			{
+				T* RuntimePtr = Cast<T>(Runtime.Get());
+				return TWeakInterfacePtr<T>(RuntimePtr);
+			}
+				
+		}
+		return TWeakInterfacePtr<T>(nullptr);
+	}
+	
+}
diff --git a/Engine/Plugins/Experimental/NNE/Source/NNXCore/Public/NNECoreRuntime.h b/Engine/Plugins/Experimental/NNE/Source/NNXCore/Public/NNECoreRuntime.h
new file mode 100644
--- /dev/null
+++ b/Engine/Plugins/Experimental/NNE/Source/NNXCore/Public/NNECoreRuntime.h
@@ -0,0 +1,26 @@
+// Copyright Epic Games, Inc. All Rights Reserved.
+
+#pragma once
+
+#include "CoreMinimal.h"
+#include "UObject/Interface.h"
+
+#include "NNECoreRuntime.generated.h"
+
+UINTERFACE()
+class NNXCORE_API UNNERuntime : public UInterface
+{
+	GENERATED_BODY()
+};
+
+class NNXCORE_API INNERuntime
+{
+	GENERATED_BODY()
+
+public:
+	virtual FString GetRuntimeName() const = 0;
+	virtual bool IsPlatformSupported(const ITargetPlatform* TargetPlatform) const = 0;
+
+	virtual bool CanCreateModelData(FString FileType, TConstArrayView<uint8> FileData) const = 0;
+	virtual TArray<uint8> CreateModelData(FString FileType, TConstArrayView<uint8> FileData) = 0;
+};
\ No newline at end of file
diff --git a/Engine/Plugins/Experimental/NNE/Source/NNXCore/Public/NNECoreRuntimeCPU.h b/Engine/Plugins/Experimental/NNE/Source/NNXCore/Public/NNECoreRuntimeCPU.h
new file mode 100644
--- /dev/null
+++ b/Engine/Plugins/Experimental/NNE/Source/NNXCore/Public/NNECoreRuntimeCPU.h
@@ -0,0 +1,64 @@
+// Copyright Epic Games, Inc. All Rights Reserved.
+
+#pragma once
+
+#include "CoreMinimal.h"
+#include "NNECoreTypes.h"
+#include "NNECoreRuntime.h"
+#include "UObject/Interface.h"
+
+#include "NNECoreRuntimeCPU.generated.h"
+
+class UNNEModelData;
+
+namespace UE::NNECore
+{
+
+struct NNXCORE_API FTensorBindingCPU
+{
+	void*	Data;
+	uint64	SizeInBytes;
+};
+
+class NNXCORE_API IModelCPU
+{
+public:
+
+	virtual ~IModelCPU() = default;
+
+	/** Getters for tensor description as defined by the model potentially with variable dimensions() */
+	virtual TConstArrayView<FTensorDesc> GetInputTensorDescs() const = 0;
+	virtual TConstArrayView<FTensorDesc> GetOutputTensorDescs() const = 0;
+
+	/** Getters for input shapes if they were set already (see SetInputTensorShapes). Empty list otherwise. */
+	virtual TConstArrayView<FTensorShape> GetInputTensorShapes() const = 0;
+
+	/** Getters for outputs shapes if they were already resolved. Empty list otherwise.
+	* Output shape might be resolved after a call to SetInputTensorShapes() if the model and engine support it
+	* otherwise they will be resolved during Run() or EnqueueRDG() */
+	virtual TConstArrayView<FTensorShape> GetOutputTensorShapes() const = 0;
+	
+	/** Prepare the model to be run with the given input shape, The call is mandatory before Run or Enqueue() can be called. */
+	virtual int SetInputTensorShapes(TConstArrayView<FTensorShape> InInputShapes) = 0;
+
+	/** This call is synchronous on all engine types(CPU, RDG), i.e.the calling thread will be blocked until execution is finished.
+	* Bindings should point to a buffers big enough */
+	virtual int RunSync(TConstArrayView<FTensorBindingCPU> InInputTensors, TConstArrayView<FTensorBindingCPU> InOutputTensors) = 0;
+};
+
+} // UE::NNECore
+
+UINTERFACE()
+class NNXCORE_API UNNERuntimeCPU : public UInterface
+{
+	GENERATED_BODY()
+};
+
+class NNXCORE_API INNERuntimeCPU
+{
+	GENERATED_BODY()
+	
+public:
+	virtual bool CanCreateModelCPU(TConstArrayView<uint8> Data) const = 0;
+	virtual TUniquePtr<UE::NNECore::IModelCPU> CreateModelCPU(TConstArrayView<uint8> Data) = 0;
+};
\ No newline at end of file
diff --git a/Engine/Plugins/Experimental/NNE/Source/NNXCore/Public/NNECoreRuntimeRDG.h b/Engine/Plugins/Experimental/NNE/Source/NNXCore/Public/NNECoreRuntimeRDG.h
new file mode 100644
--- /dev/null
+++ b/Engine/Plugins/Experimental/NNE/Source/NNXCore/Public/NNECoreRuntimeRDG.h
@@ -0,0 +1,63 @@
+// Copyright Epic Games, Inc. All Rights Reserved.
+
+#pragma once
+
+#include "CoreMinimal.h"
+#include "NNECoreTypes.h"
+#include "RenderGraphFwd.h"
+#include "UObject/Interface.h"
+
+#include "NNECoreRuntimeRDG.generated.h"
+
+class UNNEModelData;
+
+namespace UE::NNECore
+{
+
+struct NNXCORE_API FTensorBindingRDG
+{
+	FRDGBufferRef Buffer;
+};
+
+class NNXCORE_API IModelRDG
+{
+public:
+
+	virtual ~IModelRDG() = default;
+
+	/** Getters for tensor description as defined by the model potentially with variable dimensions() */
+	virtual TConstArrayView<FTensorDesc> GetInputTensorDescs() const = 0;
+	virtual TConstArrayView<FTensorDesc> GetOutputTensorDescs() const = 0;
+
+	/** Getters for input shapes if they were set already (see SetInputTensorShapes). Empty list otherwise. */
+	virtual TConstArrayView<FTensorShape> GetInputTensorShapes() const = 0;
+
+	/** Getters for outputs shapes if they were already resolved. Empty list otherwise.
+	* Output shape might be resolved after a call to SetInputTensorShapes() if the model and engine support it
+	* otherwise they will be resolved during Run() or EnqueueRDG() */
+	virtual TConstArrayView<FTensorShape> GetOutputTensorShapes() const = 0;
+	
+	/** Prepare the model to be run with the given input shape, The call is mandatory before Run or Enqueue() can be called. */
+	virtual int SetInputTensorShapes(TConstArrayView<FTensorShape> InInputShapes) = 0;
+
+	/** Enqueue the execution on the Render graph render thread. It's caller's responsibility to actually run the graph.
+	* Bindings should point to a buffers big enough */
+	virtual bool EnqueueRDG(FRDGBuilder& RDGBuilder, TConstArrayView<FTensorBindingRDG> Inputs, TConstArrayView<FTensorBindingRDG> Outputs) = 0;
+};
+
+} // UE::NNECore
+
+UINTERFACE()
+class NNXCORE_API UNNERuntimeRDG : public UInterface
+{
+	GENERATED_BODY()
+};
+
+class NNXCORE_API INNERuntimeRDG
+{
+	GENERATED_BODY()
+
+public:
+	virtual bool CanCreateModelRDG(TConstArrayView<uint8> Data) const = 0;
+	virtual TUniquePtr<UE::NNECore::IModelRDG> CreateModelRDG(TConstArrayView<uint8> Data) = 0;
+};
\ No newline at end of file
diff --git a/Engine/Plugins/Experimental/NNE/Source/NNXQA/Private/NNXQAUtils.cpp b/Engine/Plugins/Experimental/NNE/Source/NNXQA/Private/NNXQAUtils.cpp
--- a/Engine/Plugins/Experimental/NNE/Source/NNXQA/Private/NNXQAUtils.cpp
+++ b/Engine/Plugins/Experimental/NNE/Source/NNXQA/Private/NNXQAUtils.cpp
@@ -3,10 +3,14 @@
 #include "NNXQAUtils.h"
 #include "HAL/UnrealMemory.h"
 #include "Kismet/GameplayStatics.h"
+#include "NNECore.h"
 #include "NNECoreAttributeMap.h"
+#include "NNECoreRuntime.h"
+#include "NNECoreRuntimeCPU.h"
 #include "NNXCore.h"
 #include "NNXInferenceModel.h"
 #include "NNXModelOptimizerInterface.h"
+#include "UObject/WeakInterfacePtr.h"
 
 namespace NNX 
 {
@@ -108,6 +112,46 @@ namespace Test
 		return Buffer;
 	}
 
+	static void FillInputTensorBindingsCPU(TConstArrayView<UE::NNECore::Internal::FTensor> Tensors, TConstArrayView<FTests::FTensorData> TensorsData,
+		TArray<TArray<char>>& OutMemBuffers, TArray<UE::NNECore::FTensorBindingCPU>& OutBindings)
+	{
+		OutBindings.Empty();
+		OutMemBuffers.Empty();
+		check(TensorsData.Num() == 0 || Tensors.Num() == TensorsData.Num());
+
+		for (int Index = 0; Index < Tensors.Num(); ++Index)
+		{
+			const UE::NNECore::Internal::FTensor& Tensor = Tensors[Index];
+			TArray<char>& MemBuffer = OutMemBuffers.Emplace_GetRef();
+			if (TensorsData.Num() == 0 || TensorsData[Index].Num() == 0)
+			{
+				ElementWiseCosTensorInitializer Initializer(Tensor.GetDataType(), Index);
+				MemBuffer = GenerateTensorDataForTest(Tensor, Initializer);
+			}
+			else
+			{
+				MemBuffer = TensorsData[Index];
+			}
+			check(MemBuffer.Num() == Tensor.GetDataSize());
+			
+			OutBindings.Emplace(UE::NNECore::FTensorBindingCPU{(void*)MemBuffer.GetData(),(uint64)MemBuffer.Num()} );
+		}
+	}
+
+	static void FillOutputTensorBindingsCPU(TConstArrayView<UE::NNECore::Internal::FTensor> Tensors, TArray<TArray<char>>& OutMemBuffers, TArray<UE::NNECore::FTensorBindingCPU>& OutBindings)
+	{
+		OutBindings.Empty();
+		OutMemBuffers.Empty();
+		for (int Index = 0; Index < Tensors.Num(); ++Index)
+		{
+			TArray<char>& MemBuffer = OutMemBuffers.Emplace_GetRef();
+			MemBuffer.SetNumUninitialized(Tensors[Index].GetDataSize());
+			char constexpr MagicNumber = 0x5b;
+			FMemory::Memset(MemBuffer.GetData(), MagicNumber, MemBuffer.Num());
+			OutBindings.Emplace(UE::NNECore::FTensorBindingCPU{ (void*)MemBuffer.GetData(),(uint64)MemBuffer.Num() });
+		}
+	}
+
 	static void FillInputTensorBindings(TConstArrayView<UE::NNECore::Internal::FTensor> Tensors, TConstArrayView<FTests::FTensorData> TensorsData,
 		TArray<TArray<char>>& OutMemBuffers, TArray<FMLTensorBinding>& OutBindings)
 	{
@@ -369,6 +413,109 @@ namespace Test
 		return true;
 	}
 
+	TUniquePtr<UE::NNECore::IModelCPU> GetInferenceModel(const FNNIModelRaw& ONNXModelData, const FString& RuntimeName)
+	{
+		TWeakInterfacePtr<INNERuntime>    RuntimeForModelData      = UE::NNECore::GetRuntime<INNERuntime>(RuntimeName);
+		TWeakInterfacePtr<INNERuntimeCPU> RuntimeForInferenceModel = UE::NNECore::GetRuntime<INNERuntimeCPU>(RuntimeName);
+		if (!RuntimeForModelData.IsValid() || !RuntimeForInferenceModel.IsValid())
+		{
+			UE_LOG(LogNNE, Error, TEXT("Can't get %s runtime."), *RuntimeName);
+			return TUniquePtr<UE::NNECore::IModelCPU>();
+		}
+
+		TArray<uint8> Data = RuntimeForModelData->CreateModelData(FString("onnx"), ONNXModelData.Data);
+		TUniquePtr<UE::NNECore::IModelCPU> InferenceModel = RuntimeForInferenceModel->CreateModelCPU(Data);
+		
+		return InferenceModel;
+	}
+	
+	bool RunTestInferenceCPU(const FNNIModelRaw& ONNXModelData, const FTests::FTestSetup& TestSetup,
+		const FString& RuntimeName, TArray<UE::NNECore::Internal::FTensor>& OutOutputTensors, TArray<TArray<char>>& OutOutputMemBuffers)
+	{
+		OutOutputMemBuffers.Empty();
+
+		TUniquePtr<UE::NNECore::IModelCPU> InferenceModel = GetInferenceModel(ONNXModelData, RuntimeName);
+
+		if (!InferenceModel.IsValid())
+		{
+			UE_LOG(LogNNX, Error, TEXT("Could not create Inference model."));
+			return false;
+		}
+
+		// If test does not ask for specific input/output, fill with model default converting variable dim to 1 if any.
+		TArray<UE::NNECore::Internal::FTensor> InputTensors;
+
+		FillTensors(TestSetup.Inputs, InferenceModel->GetInputTensorDescs(), InputTensors);
+		FillTensors(TestSetup.Outputs, InferenceModel->GetOutputTensorDescs(), OutOutputTensors);
+
+		// bind tensors to memory (CPU) and initialize
+		TArray<UE::NNECore::FTensorBindingCPU> InputBindings;
+		TArray<UE::NNECore::FTensorBindingCPU> OutputBindings;
+		TArray<TArray<char>> InputMemBuffers;
+
+		FillInputTensorBindingsCPU(InputTensors, TestSetup.InputsData, InputMemBuffers, InputBindings);
+		FillOutputTensorBindingsCPU(OutOutputTensors, OutOutputMemBuffers, OutputBindings);
+
+		// To help for debugging sessions.
+#ifdef UE_BUILD_DEBUG
+		const constexpr int32 NumTensorPtrForDebug = 3;
+		float* InputsAsFloat[NumTensorPtrForDebug];
+		float* OutputsAsFloat[NumTensorPtrForDebug];
+		FMemory::Memzero(InputsAsFloat, NumTensorPtrForDebug * sizeof(float*));
+		FMemory::Memzero(OutputsAsFloat, NumTensorPtrForDebug * sizeof(float*));
+		for (int32 i = 0; i < NumTensorPtrForDebug && i < InputMemBuffers.Num(); ++i)
+		{
+			InputsAsFloat[i] = (float*)InputMemBuffers[i].GetData();
+		}
+		for (int32 i = 0; i < NumTensorPtrForDebug && i < OutOutputMemBuffers.Num(); ++i)
+		{
+			OutputsAsFloat[i] = (float*)OutOutputMemBuffers[i].GetData();
+		}
+#endif //UE_BUILD_DEBUG
+
+		TArray<UE::NNECore::FTensorShape> InputShapes;
+		for (const UE::NNECore::Internal::FTensor& Tensor : InputTensors)
+		{
+			InputShapes.Emplace(Tensor.GetShape());
+		}
+
+		// Setup inputs
+		if (0 != InferenceModel->SetInputTensorShapes(InputShapes))
+		{
+			UE_LOG(LogNNX, Error, TEXT("Failed to set input tensor shapes."));
+			return false;
+		}
+
+		// Run inference
+		if (0 != InferenceModel->RunSync(InputBindings, OutputBindings))
+		{
+			UE_LOG(LogNNX, Error, TEXT("Failed to run the model."));
+			return false;
+		}
+
+		// Verify output shapes are as expected
+		TConstArrayView<UE::NNECore::FTensorShape> OutputShapes = InferenceModel->GetOutputTensorShapes();
+		if (OutputShapes.Num() != OutOutputTensors.Num())
+		{
+			UE_LOG(LogNNX, Error, TEXT("Expected %d output tensors, got %d."), OutOutputTensors.Num(), OutputShapes.Num());
+			return false;
+		}
+		for (int i = 0; i < OutOutputTensors.Num(); ++i)
+		{
+			const UE::NNECore::FTensorShape& RefTensorShape = OutOutputTensors[i].GetShape();
+			const UE::NNECore::FTensorShape& OtherTensorShape = OutputShapes[i];
+			if (RefTensorShape != OtherTensorShape)
+			{
+				TArrayView<const uint32> RefShape(RefTensorShape.GetData());
+				TArrayView<const uint32> OtherShape(OtherTensorShape.GetData());
+				UE_LOG(LogNNX, Error, TEXT("Output shape do not match at index %d.\nExpected: %s\nGot:      %s"), i, *ShapeToString(RefShape), *ShapeToString(OtherShape));
+				return false;
+			}
+		}
+
+		return true;
+	}
+
 	static bool RunTestInference(const FNNIModelRaw& ONNXModelData, const FTests::FTestSetup& TestSetup, 
 		IRuntime* Runtime, TArray<UE::NNECore::Internal::FTensor>& OutOutputTensors, TArray<TArray<char>>& OutOutputMemBuffers)
 	{
@@ -502,20 +649,14 @@ namespace Test
 		UE_LOG(LogNNX, Display, TEXT("Starting tests of '%s'"), *TestSetup.TargetName);
 
 		// Reference runtime
-		IRuntime* RefRuntime = NNX::GetRuntime(TEXT("NNXRuntimeCPU"));
-		if (!RefRuntime)
-		{
-			UE_LOG(LogNNX, Error, TEXT("Can't load NNXRuntimeCPU runtime. Tests ABORTED!"));
-			return false;
-		}
-		const FString& RefName = RefRuntime->GetRuntimeName();
+		const FString& RefName = TEXT("NNERuntimeCPU");
 		const float AbsoluteTolerance = TestSetup.GetAbsoluteToleranceForRuntime(RefName);
 		const float RelativeTolerance = TestSetup.GetRelativeToleranceForRuntime(RefName);
 		TArray<TArray<char>> RefOutputMemBuffers;
 		TArray<UE::NNECore::Internal::FTensor> RefOutputTensors;
 		bool bAllTestsSucceeded = true;
 
-		if (!RunTestInference(ONNXModel, TestSetup, RefRuntime, RefOutputTensors, RefOutputMemBuffers))
+		if (!RunTestInferenceCPU(ONNXModel, TestSetup, RefName, RefOutputTensors, RefOutputMemBuffers))
 		{
 			UE_LOG(LogNNX, Error, TEXT("Error running reference inference with engine %s."), *RefName);
 			return false;
diff --git a/Engine/Plugins/Experimental/NNE/Source/NNXRuntimeCPU/NNXRuntimeCPU.Build.cs b/Engine/Plugins/Experimental/NNE/Source/NNXRuntimeCPU/NNXRuntimeCPU.Build.cs
--- a/Engine/Plugins/Experimental/NNE/Source/NNXRuntimeCPU/NNXRuntimeCPU.Build.cs
+++ b/Engine/Plugins/Experimental/NNE/Source/NNXRuntimeCPU/NNXRuntimeCPU.Build.cs
@@ -28,7 +28,8 @@ public class NNXRuntimeCPU : ModuleRules
 			(
 			new string[] {
 				//"ORTDefault",
-                "NNXUtils",
+				"CoreUObject",
+				"NNXUtils",
 				"NNX_ONNXRuntime"	// Select this for open-source-based ONNX Runtime
 			}
 		);
diff --git a/Engine/Plugins/Experimental/NNE/Source/NNXRuntimeCPU/Private/Module.cpp b/Engine/Plugins/Experimental/NNE/Source/NNXRuntimeCPU/Private/Module.cpp
deleted file mode 100644
diff --git a/Engine/Plugins/Experimental/NNE/Source/NNXRuntimeCPU/Private/Module.h b/Engine/Plugins/Experimental/NNE/Source/NNXRuntimeCPU/Private/Module.h
deleted file mode 100644
diff --git a/Engine/Plugins/Experimental/NNE/Source/NNXRuntimeCPU/Private/NNERuntimeCPU.cpp b/Engine/Plugins/Experimental/NNE/Source/NNXRuntimeCPU/Private/NNERuntimeCPU.cpp
new file mode 100644
--- /dev/null
+++ b/Engine/Plugins/Experimental/NNE/Source/NNXRuntimeCPU/Private/NNERuntimeCPU.cpp
@@ -0,0 +1,84 @@
+// Copyright Epic Games, Inc. All Rights Reserved.
+
+#include "NNERuntimeCPU.h"
+#include "NNERuntimeCPUModel.h"
+#include "NNERuntimeCPUUtils.h"
+#include "NNXModelOptimizer.h"
+#include "NNXModelOptimizerInterface.h"
+#include "NNECoreAttributeMap.h"
+#include "NNECoreModelData.h"
+#include "NNEProfilingTimer.h"
+#include "RedirectCoutAndCerrToUeLog.h"
+
+using namespace NNX;
+
+FGuid UNNERuntimeCPUImpl::GUID = FGuid((int32)'R', (int32)'C', (int32)'P', (int32)'U');
+int32 UNNERuntimeCPUImpl::Version = 0x00000001;
+
+bool UNNERuntimeCPUImpl::CanCreateModelData(FString FileType, TConstArrayView<uint8> FileData) const
+{
+	return FileType.Compare("onnx", ESearchCase::IgnoreCase) == 0;
+}
+
+TArray<uint8> UNNERuntimeCPUImpl::CreateModelData(FString FileType, TConstArrayView<uint8> FileData)
+{
+	if (!CanCreateModelData(FileType, FileData))
+	{
+		return {};
+	}
+
+	TUniquePtr<IModelOptimizer> Optimizer = CreateONNXToONNXModelOptimizer();
+
+	FNNIModelRaw InputModel;
+	InputModel.Data = FileData;
+	InputModel.Format = ENNXInferenceFormat::ONNX;
+	FNNIModelRaw OutputModel;
+	FOptimizerOptionsMap Options;
+	if (!Optimizer->Optimize(InputModel, OutputModel, Options))
+	{
+		return {};
+	}
+
+	int32 GuidSize = sizeof(UNNERuntimeCPUImpl::GUID);
+	int32 VersionSize = sizeof(UNNERuntimeCPUImpl::Version);
+	TArray<uint8> Result;
+	FMemoryWriter Writer(Result);
+	Writer << UNNERuntimeCPUImpl::GUID;
+	Writer << UNNERuntimeCPUImpl::Version;
+	Writer.Serialize(OutputModel.Data.GetData(), OutputModel.Data.Num());
+	return Result;
+}
+
+bool UNNERuntimeCPUImpl::CanCreateModelCPU(TConstArrayView<uint8> Data) const
+{
+	int32 GuidSize = sizeof(UNNERuntimeCPUImpl::GUID);
+	int32 VersionSize = sizeof(UNNERuntimeCPUImpl::Version);
+	
+	if (Data.Num() <= GuidSize + VersionSize)
+	{
+		return false;
+	}
+	
+	bool bResult = FGenericPlatformMemory::Memcmp(&(Data[0]), &(UNNERuntimeCPUImpl::GUID), GuidSize) == 0;
+	bResult &= FGenericPlatformMemory::Memcmp(&(Data[GuidSize]), &(UNNERuntimeCPUImpl::Version), VersionSize) == 0;
+	return bResult;
+}
+
+TUniquePtr<UE::NNECore::IModelCPU> UNNERuntimeCPUImpl::CreateModelCPU(TConstArrayView<uint8> Data)
+{
+	if (!CanCreateModelCPU(Data))
+	{
+		return TUniquePtr<UE::NNECore::IModelCPU>();
+	}
+
+	const UE::NNERuntimeCPU::Private::FRuntimeConf InConf;
+	UE::NNERuntimeCPU::Private::FModelCPU* Model = new UE::NNERuntimeCPU::Private::FModelCPU(&NNEEnvironmentCPU, InConf);
+	
+	if (!Model->Init(Data))
+	{
+		delete Model;
+		return TUniquePtr<UE::NNECore::IModelCPU>();
+	}
+	UE::NNECore::IModelCPU* IModel = static_cast<UE::NNECore::IModelCPU*>(Model);
+	return TUniquePtr<UE::NNECore::IModelCPU>(IModel);
+}
\ No newline at end of file
diff --git a/Engine/Plugins/Experimental/NNE/Source/NNXRuntimeCPU/Private/NNERuntimeCPU.h b/Engine/Plugins/Experimental/NNE/Source/NNXRuntimeCPU/Private/NNERuntimeCPU.h
new file mode 100644
--- /dev/null
+++ b/Engine/Plugins/Experimental/NNE/Source/NNXRuntimeCPU/Private/NNERuntimeCPU.h
@@ -0,0 +1,42 @@
+// Copyright Epic Games, Inc. All Rights Reserved.
+
+#pragma once
+
+#include "CoreMinimal.h"
+#include "NNECoreRuntime.h"
+#include "NNECoreRuntimeCPU.h"
+#include "UObject/Class.h"
+#include "UObject/Object.h"
+#include "UObject/UObjectBaseUtility.h"
+
+#include "NNXThirdPartyWarningDisabler.h"
+NNX_THIRD_PARTY_INCLUDES_START
+#undef check
+#undef TEXT
+#include "core/session/onnxruntime_cxx_api.h"
+NNX_THIRD_PARTY_INCLUDES_END
+
+#include "NNERuntimeCPU.generated.h"
+
+UCLASS()
+class UNNERuntimeCPUImpl : public UObject, public INNERuntime, public INNERuntimeCPU
+{
+	GENERATED_BODY()
+
+public:
+	static FGuid GUID;
+	static int32 Version;
+
+	Ort::Env NNEEnvironmentCPU;
+	UNNERuntimeCPUImpl() {};
+	virtual ~UNNERuntimeCPUImpl() {}
+		
+	virtual FString GetRuntimeName() const override { return TEXT("NNERuntimeCPU"); };
+	virtual bool IsPlatformSupported(const ITargetPlatform* TargetPlatform) const override { return true; };
+
+	virtual bool CanCreateModelData(FString FileType, TConstArrayView<uint8> FileData) const override;
+	virtual TArray<uint8> CreateModelData(FString FileType, TConstArrayView<uint8> FileData) override;
+
+	virtual bool CanCreateModelCPU(TConstArrayView<uint8> Data) const override;
+	virtual TUniquePtr<UE::NNECore::IModelCPU> CreateModelCPU(TConstArrayView<uint8> Data) override;
+};
\ No newline at end of file
diff --git a/Engine/Plugins/Experimental/NNE/Source/NNXRuntimeCPU/Private/NNERuntimeCPUModel.cpp b/Engine/Plugins/Experimental/NNE/Source/NNXRuntimeCPU/Private/NNERuntimeCPUModel.cpp
new file mode 100644
--- /dev/null
+++ b/Engine/Plugins/Experimental/NNE/Source/NNXRuntimeCPU/Private/NNERuntimeCPUModel.cpp
@@ -0,0 +1,301 @@
+// Copyright Epic Games, Inc. All Rights Reserved.
+#include "NNERuntimeCPUModel.h"
+#include "NNERuntimeCPU.h"
+#include "NNERuntimeCPUUtils.h"
+#include "NNEProfilingTimer.h"
+#include "RedirectCoutAndCerrToUeLog.h"
+
+namespace UE::NNERuntimeCPU::Private
+{
+
+	FModelCPU::FModelCPU() :
+		bIsLoaded(false),
+		bHasRun(false)
+	{ }
+
+	FModelCPU::FModelCPU(
+		Ort::Env* InORTEnvironment,
+		const FRuntimeConf& InRuntimeConf) :
+		bIsLoaded(false),
+		bHasRun(false),
+		RuntimeConf(InRuntimeConf),
+		ORTEnvironment(InORTEnvironment)
+	{ }
+
+	bool FModelCPU::Init(TConstArrayView<uint8> ModelData)
+	{
+		DECLARE_SCOPE_CYCLE_COUNTER(TEXT("FModelCPU_Init"), STAT_FModelCPU_Init, STATGROUP_NNE);
+		
+		// Get the header size
+		int32 GuidSize = sizeof(UNNERuntimeCPUImpl::GUID);
+		int32 VersionSize = sizeof(UNNERuntimeCPUImpl::Version);
+
+		// Clean previous networks
+		bIsLoaded = false;
+		TConstArrayView<uint8> ModelBuffer = TConstArrayView<uint8>(&(ModelData.GetData()[GuidSize + VersionSize]), ModelData.Num() - GuidSize - VersionSize);
+
+		// Checking Inference Model 
+		{
+			if (ModelBuffer.Num() == 0)
+			{
+				UE_LOG(LogNNE, Warning, TEXT("FModelCPU::Load(): Input model data is empty."));
+				return false;
+			}
+
+		}
+
+#if WITH_EDITOR
+		try
+#endif //WITH_EDITOR
+		{
+			const FRedirectCoutAndCerrToUeLog RedirectCoutAndCerrToUeLog;
+			if (!InitializedAndConfigureMembers())
+			{
+				UE_LOG(LogNNE, Warning, TEXT("Load(): InitializedAndConfigureMembers failed."));
+				return false;
+			}
+
+			{
+				DECLARE_SCOPE_CYCLE_COUNTER(TEXT("FModelCPU_Init_CreateORTSession"), STAT_FModelCPU_Init_CreateORTSession, STATGROUP_NNE);
+				// Read model from InferenceModel
+				Session = MakeUnique<Ort::Session>(*ORTEnvironment, ModelBuffer.GetData(), ModelBuffer.Num(), *SessionOptions);
+			}
+
+			if (!ConfigureTensors(true))
+			{
+				UE_LOG(LogNNE, Warning, TEXT("Load(): Failed to configure Inputs tensors."));
+				return false;
+			}
+
+			if (!ConfigureTensors(false))
+			{
+				UE_LOG(LogNNE, Warning, TEXT("Load(): Failed to configure Outputs tensors."));
+				return false;
+			}
+		}
+#if WITH_EDITOR
+		catch (const std::exception& Exception)
+		{
+			UE_LOG(LogNNE, Error, TEXT("%s"), UTF8_TO_TCHAR(Exception.what()));
+			return false;
+		}
+#endif //WITH_EDITOR
+
+		bIsLoaded = true;
+
+		// Reset Stats
+		ResetStats();
+
+		return IsLoaded();
+	}
+
+	bool FModelCPU::IsLoaded() const
+	{
+		return bIsLoaded;
+	}
+
+	bool FModelCPU::InitializedAndConfigureMembers()
+	{
+		// Initialize 
+		// Set up ORT and create an environment
+		Allocator = MakeUnique<Ort::AllocatorWithDefaultOptions>();
+		AllocatorInfo = MakeUnique<Ort::MemoryInfo>(Ort::MemoryInfo::CreateCpu(OrtDeviceAllocator, OrtMemTypeCPU));
+
+		// Configure Session
+		SessionOptions = MakeUnique<Ort::SessionOptions>();
+
+		// Configure number threads
+		//SessionOptions->SetIntraOpNumThreads(RuntimeConf.NumberOfThreads);
+		// Uncomment if you want to change the priority of the threads, by default is TPri_Normal
+		//SessionOptions->SetPriorityOpThreads(RuntimeConf.ThreadPriority);
+
+		// ORT CPU
+		SessionOptions->SetGraphOptimizationLevel(RuntimeConf.OptimizationLevel); // ORT_ENABLE_ALL, ORT_ENABLE_EXTENDED, ORT_ENABLE_BASIC, ORT_DISABLE_ALL
+		SessionOptions->EnableCpuMemArena();
+
+		return true;
+	}
+
+
+	bool FModelCPU::ConfigureTensors(const bool InIsInput)
+	{
+		const bool bIsInput = InIsInput;
+
+		const uint32 NumberTensors = bIsInput ? Session->GetInputCount() : Session->GetOutputCount();
+
+		TArray<NNECore::FTensorDesc>& SymbolicTensorDescs = bIsInput ? InputSymbolicTensors : OutputSymbolicTensors;
+		TArray<ONNXTensorElementDataType>& TensorsORTType = bIsInput ? InputTensorsORTType : OutputTensorsORTType;
+		TArray<const char*>& TensorNames = bIsInput ? InputTensorNames : OutputTensorNames;
+
+		for (uint32 TensorIndex = 0; TensorIndex < NumberTensors; ++TensorIndex)
+		{
+			// Get Tensor name
+			const char* CurTensorName = bIsInput ? Session->GetInputName(TensorIndex, *Allocator) : Session->GetOutputName(TensorIndex, *Allocator);
+			TensorNames.Emplace(CurTensorName);
+
+			// Get node type
+			Ort::TypeInfo CurrentTypeInfo = bIsInput ? Session->GetInputTypeInfo(TensorIndex) : Session->GetOutputTypeInfo(TensorIndex);
+			Ort::TensorTypeAndShapeInfo CurrentTensorInfo = CurrentTypeInfo.GetTensorTypeAndShapeInfo();
+			const ONNXTensorElementDataType ONNXTensorElementDataTypeEnum = CurrentTensorInfo.GetElementType();
+			CurrentTypeInfo.release();
+
+			TensorsORTType.Emplace(ONNXTensorElementDataTypeEnum);
+
+			std::pair<ENNETensorDataType, uint64> TypeAndSize = TranslateTensorTypeORTToNNE(ONNXTensorElementDataTypeEnum);
+
+			TArray<int32> ShapeData;
+			ShapeData.Reserve(CurrentTensorInfo.GetShape().size());
+			for (const int64_t CurrentTensorSize : CurrentTensorInfo.GetShape())
+			{
+				ShapeData.Add((int32)CurrentTensorSize);
+			}
+
+			NNECore::FSymbolicTensorShape Shape = NNECore::FSymbolicTensorShape::Make(ShapeData);
+			NNECore::FTensorDesc SymbolicTensorDesc = NNECore::FTensorDesc::Make(FString(CurTensorName), Shape, TypeAndSize.first);
+
+			check(SymbolicTensorDesc.GetElemByteSize() == TypeAndSize.second);
+			SymbolicTensorDescs.Emplace(SymbolicTensorDesc);
+		}
+
+		return true;
+	}
+	
+	int FModelCPU::SetInputTensorShapes(TConstArrayView<NNECore::FTensorShape> InInputShapes)
+	{
+		InputTensors.Empty();
+		OutputTensors.Empty();
+		OutputTensorShapes.Empty();
+
+		// Verify input shape are valid for the model and set InputTensorShapes
+		if (FModelBase<IModelCPU>::SetInputTensorShapes(InInputShapes) != 0)
+		{
+			return -1;
+		}
+
+		// Setup concrete input tensor
+		for (int i = 0; i < InputSymbolicTensors.Num(); ++i)
+		{
+			NNECore::Internal::FTensor Tensor = NNECore::Internal::FTensor::Make(InputSymbolicTensors[i].GetName(), InInputShapes[i], InputSymbolicTensors[i].GetDataType());
+			InputTensors.Emplace(Tensor);
+		}
+
+		// Here model optimization could be done now that we know the input shapes, for some models
+		// this would allow to resolve output shapes here rather than during inference.
+
+		// Setup concrete output shapes only if all model output shapes are concretes, otherwise it will be set during Run()
+		for (NNECore::FTensorDesc SymbolicTensorDesc : OutputSymbolicTensors)
+		{
+			if (SymbolicTensorDesc.GetShape().IsConcrete())
+			{
+				NNECore::Internal::FTensor Tensor = NNECore::Internal::FTensor::MakeFromSymbolicDesc(SymbolicTensorDesc);
+				OutputTensors.Emplace(Tensor);
+				OutputTensorShapes.Emplace(Tensor.GetShape());
+			}
+		}
+		if (OutputTensors.Num() != OutputSymbolicTensors.Num())
+		{
+			OutputTensors.Empty();
+			OutputTensorShapes.Empty();
+		}
+
+		return 0;
+	}
+
+	int FModelCPU::RunSync(TConstArrayView<NNECore::FTensorBindingCPU> InInputBindings, TConstArrayView<NNECore::FTensorBindingCPU> InOutputBindings)
+	{
+		DECLARE_SCOPE_CYCLE_COUNTER(TEXT("FModelCPU_Run"), STAT_FModelCPU_Run, STATGROUP_NNE);
+
+		// Sanity check
+		if (!bIsLoaded)
+		{
+			UE_LOG(LogNNE, Warning, TEXT("FModelCPU::Run(): Call FModelCPU::Load() to load a model first."));
+			return -1;
+		}
+
+		// Verify the model inputs were prepared
+		if (InputTensorShapes.Num() == 0)
+		{
+			UE_LOG(LogNNE, Error, TEXT("Run(): Input shapes are not set, please call SetInputTensorShapes."));
+			return -1;
+		}
+
+		NNEProfiling::Internal::FTimer RunTimer;
+		RunTimer.Tic();
+
+		if (!bHasRun)
+		{
+			bHasRun = true;
+		}
+
+#if WITH_EDITOR
+		try
+#endif //WITH_EDITOR
+		{
+			TArray<Ort::Value> InputOrtTensors;
+			BindTensorsToORT(InInputBindings, InputTensors, InputTensorsORTType, AllocatorInfo.Get(), InputOrtTensors);
+
+			if (!OutputTensors.IsEmpty())
+			{
+				// If output shapes are known we can directly map preallocated output buffers
+				TArray<Ort::Value> OutputOrtTensors;
+				BindTensorsToORT(InOutputBindings, OutputTensors, OutputTensorsORTType, AllocatorInfo.Get(), OutputOrtTensors);
+
+				Session->Run(Ort::RunOptions{ nullptr },
+					InputTensorNames.GetData(), &InputOrtTensors[0], InputTensorNames.Num(),
+					OutputTensorNames.GetData(), &OutputOrtTensors[0], OutputTensorNames.Num());
+			}
+			else
+			{
+				TArray<Ort::Value> OutputOrtTensors;
+				for (int i = 0; i < InOutputBindings.Num(); ++i)
+				{
+					OutputOrtTensors.Emplace(nullptr);
+				}
+
+				Session->Run(Ort::RunOptions{ nullptr },
+					InputTensorNames.GetData(), &InputOrtTensors[0], InputTensorNames.Num(),
+					OutputTensorNames.GetData(), &OutputOrtTensors[0], OutputTensorNames.Num());
+
+				// Output shapes were resolved during inference: Copy the data back to bindings and expose output tensor shapes
+				CopyFromORTToBindings(OutputOrtTensors, InOutputBindings, OutputSymbolicTensors, OutputTensors);
+				check(OutputTensorShapes.IsEmpty());
+				for (int i = 0; i < OutputTensors.Num(); ++i)
+				{
+					OutputTensorShapes.Emplace(OutputTensors[i].GetShape());
+				}
+			}
+		}
+#if WITH_EDITOR
+		catch (const std::exception& Exception)
+		{
+			UE_LOG(LogNNE, Error, TEXT("%s"), UTF8_TO_TCHAR(Exception.what()));
+		}
+#endif //WITH_EDITOR
+
+		RunStatisticsEstimator.StoreSample(RunTimer.Toc());
+
+		return 0;
+	}
+
+	float FModelCPU::GetLastRunTimeMSec() const
+	{
+		return RunStatisticsEstimator.GetLastSample();
+	}
+
+	NNEProfiling::Internal::FStatistics FModelCPU::GetRunStatistics() const
+	{
+		return RunStatisticsEstimator.GetStats();
+	}
+
+	NNEProfiling::Internal::FStatistics FModelCPU::GetInputMemoryTransferStats() const
+	{
+		return InputTransferStatisticsEstimator.GetStats();
+	}
+
+	void FModelCPU::ResetStats()
+	{
+		RunStatisticsEstimator.ResetStats();
+		InputTransferStatisticsEstimator.ResetStats();
+	}
+
+} // namespace UE::NNERuntimeCPU::Private
\ No newline at end of file
diff --git a/Engine/Plugins/Experimental/NNE/Source/NNXRuntimeCPU/Private/NNERuntimeCPUModel.h b/Engine/Plugins/Experimental/NNE/Source/NNXRuntimeCPU/Private/NNERuntimeCPUModel.h
new file mode 100644
--- /dev/null
+++ b/Engine/Plugins/Experimental/NNE/Source/NNXRuntimeCPU/Private/NNERuntimeCPUModel.h
@@ -0,0 +1,81 @@
+// Copyright Epic Games, Inc. All Rights Reserved.
+
+#pragma once
+
+#include "CoreMinimal.h"
+#include "NNECoreModelBase.h"
+#include "NNECoreRuntimeCPU.h"
+#include "NNECoreTensor.h"
+#include "NNECoreTypes.h"
+#include "NNEProfilingStatistics.h"
+
+#include "NNXThirdPartyWarningDisabler.h"
+NNX_THIRD_PARTY_INCLUDES_START
+#undef check
+#undef TEXT
+#include "core/session/onnxruntime_cxx_api.h"
+NNX_THIRD_PARTY_INCLUDES_END
+
+namespace UE::NNERuntimeCPU::Private
+{
+	struct FRuntimeConf
+	{
+		uint32 NumberOfThreads = 2;
+		GraphOptimizationLevel OptimizationLevel = GraphOptimizationLevel::ORT_ENABLE_ALL;
+		EThreadPriority ThreadPriority = EThreadPriority::TPri_Normal;
+	};
+
+	class FModelCPU : public NNECore::Internal::FModelBase<NNECore::IModelCPU>
+	{
+
+	public:
+		FModelCPU();
+		FModelCPU(Ort::Env* InORTEnvironment, const FRuntimeConf& InRuntimeConf);
+		virtual ~FModelCPU() {};
+
+		virtual int SetInputTensorShapes(TConstArrayView<NNECore::FTensorShape> InInputShapes) override;
+		virtual int RunSync(TConstArrayView<NNECore::FTensorBindingCPU> InInputBindings, TConstArrayView<NNECore::FTensorBindingCPU> InOutputBindings) override;
+
+		bool Init(TConstArrayView<uint8> ModelData);
+		bool IsLoaded() const;
+		float GetLastRunTimeMSec() const;
+		NNEProfiling::Internal::FStatistics GetRunStatistics() const;
+		NNEProfiling::Internal::FStatistics GetInputMemoryTransferStats() const;
+		void ResetStats();
+
+	private:
+
+		bool bIsLoaded = false;
+		bool bHasRun = false;
+
+		FRuntimeConf RuntimeConf;
+
+		/** ORT-related variables */
+		Ort::Env* ORTEnvironment;
+		TUniquePtr<Ort::Session> Session;
+		TUniquePtr<Ort::AllocatorWithDefaultOptions> Allocator;
+		TUniquePtr<Ort::SessionOptions> SessionOptions;
+		TUniquePtr<Ort::MemoryInfo> AllocatorInfo;
+
+		/** IO ORT-related variables */
+		TArray<ONNXTensorElementDataType> InputTensorsORTType;
+		TArray<ONNXTensorElementDataType> OutputTensorsORTType;
+
+		TArray<const char*> InputTensorNames;
+		TArray<const char*> OutputTensorNames;
+
+		TArray<NNECore::Internal::FTensor> InputTensors;
+		TArray<NNECore::Internal::FTensor> OutputTensors;
+
+		bool InitializedAndConfigureMembers();
+		bool ConfigureTensors(const bool InIsInput);
+
+		/**
+		 * Statistics-related members used for GetLastRunTimeMSec(), GetRunStatistics(), GetInputMemoryTransferStats(), ResetStats().
+		 */
+		UE::NNEProfiling::Internal::FStatisticsEstimator RunStatisticsEstimator;
+		UE::NNEProfiling::Internal::FStatisticsEstimator InputTransferStatisticsEstimator;
+
+	};
+	
+} // UE::NNERuntimeCPU::Private
\ No newline at end of file
diff --git a/Engine/Plugins/Experimental/NNE/Source/NNXRuntimeCPU/Private/NNERuntimeCPUModule.cpp b/Engine/Plugins/Experimental/NNE/Source/NNXRuntimeCPU/Private/NNERuntimeCPUModule.cpp
new file mode 100644
--- /dev/null
+++ b/Engine/Plugins/Experimental/NNE/Source/NNXRuntimeCPU/Private/NNERuntimeCPUModule.cpp
@@ -0,0 +1,51 @@
+// Copyright Epic Games, Inc. All Rights Reserved.
+
+#include "NNERuntimeCPUModule.h"
+#include "NNECore.h"
+#include "NNERuntimeCPU.h"
+#include "NNXCore.h"
+#include "NNXRuntimeCPU.h"
+#include "UObject/WeakInterfacePtr.h"
+
+void FNNERuntimeCPUModule::StartupModule()
+{
+	// NNX runtime startup
+	CPURuntime = NNX::FRuntimeCPUStartup();
+	if (CPURuntime)
+	{
+		NNX::RegisterRuntime(CPURuntime);
+	}
+
+	// NNE runtime startup
+	NNERuntimeCPU = NewObject<UNNERuntimeCPUImpl>();
+	if (NNERuntimeCPU.IsValid())
+	{
+		TWeakInterfacePtr<INNERuntime> RuntimeCPUInterface(NNERuntimeCPU.Get());
+		
+		NNERuntimeCPU->AddToRoot();
+		UE::NNECore::RegisterRuntime(RuntimeCPUInterface);
+	}
+}
+
+void FNNERuntimeCPUModule::ShutdownModule()
+{
+	// NNX runtime shutdown
+	if (CPURuntime)
+	{
+		NNX::UnregisterRuntime(CPURuntime);
+		CPURuntime = nullptr;
+	}
+	NNX::FRuntimeCPUShutdown();
+
+	// NNE runtime shutdown
+	if (NNERuntimeCPU.IsValid())
+	{
+		TWeakInterfacePtr<INNERuntime> RuntimeCPUInterface(NNERuntimeCPU.Get());
+		
+		UE::NNECore::UnregisterRuntime(RuntimeCPUInterface);
+		NNERuntimeCPU->RemoveFromRoot();
+		NNERuntimeCPU = TWeakObjectPtr<UNNERuntimeCPUImpl>(nullptr);
+	}
+}
+
+IMPLEMENT_MODULE(FNNERuntimeCPUModule, NNXRuntimeCPU);
diff --git a/Engine/Plugins/Experimental/NNE/Source/NNXRuntimeCPU/Private/NNERuntimeCPUModule.h b/Engine/Plugins/Experimental/NNE/Source/NNXRuntimeCPU/Private/NNERuntimeCPUModule.h
new file mode 100644
--- /dev/null
+++ b/Engine/Plugins/Experimental/NNE/Source/NNXRuntimeCPU/Private/NNERuntimeCPUModule.h
@@ -0,0 +1,23 @@
+// Copyright Epic Games, Inc. All Rights Reserved.
+
+#pragma once
+
+#include "CoreMinimal.h"
+#include "Modules/ModuleManager.h"
+#include "UObject/WeakObjectPtr.h"
+#include "UObject/WeakObjectPtrTemplates.h"
+
+namespace NNX { class IRuntime; }
+class UNNERuntimeCPUImpl;
+
+class FNNERuntimeCPUModule : public IModuleInterface
+{
+
+public:
+	NNX::IRuntime* CPURuntime{ nullptr };
+	TWeakObjectPtr<UNNERuntimeCPUImpl> NNERuntimeCPU{ nullptr };
+
+	// Begin IModuleInterface
+	virtual void StartupModule() override;
+	virtual void ShutdownModule() override;
+};
\ No newline at end of file
diff --git a/Engine/Plugins/Experimental/NNE/Source/NNXRuntimeCPU/Private/NNERuntimeCPUUtils.h b/Engine/Plugins/Experimental/NNE/Source/NNXRuntimeCPU/Private/NNERuntimeCPUUtils.h
new file mode 100644
--- /dev/null
+++ b/Engine/Plugins/Experimental/NNE/Source/NNXRuntimeCPU/Private/NNERuntimeCPUUtils.h
@@ -0,0 +1,214 @@
+// Copyright Epic Games, Inc. All Rights Reserved.
+
+#pragma once
+
+#include "HAL/UnrealMemory.h"
+#include "NNECore.h"
+#include "NNECoreRuntimeCPU.h"
+#include "NNECoreTypes.h"
+#include "NNECoreTensor.h"
+#include "NNXThirdPartyWarningDisabler.h"
+NNX_THIRD_PARTY_INCLUDES_START
+#undef check
+#undef TEXT
+#include "core/session/onnxruntime_cxx_api.h"
+NNX_THIRD_PARTY_INCLUDES_END
+
+DECLARE_STATS_GROUP(TEXT("NNE"), STATGROUP_NNE, STATCAT_Advanced);
+
+namespace UE::NNERuntimeCPU::Private
+{
+
+	using TypeInfoORT = std::pair<ENNETensorDataType, uint64>;
+
+	inline TypeInfoORT TranslateTensorTypeORTToNNE(unsigned int OrtDataType)
+	{
+		ENNETensorDataType DataType = ENNETensorDataType::None;
+		uint64 ElementSize = 0;
+
+		switch (OrtDataType) {
+		case ONNX_TENSOR_ELEMENT_DATA_TYPE_UNDEFINED:
+			DataType = ENNETensorDataType::None;
+			ElementSize = 0;
+			break;
+
+		case ONNX_TENSOR_ELEMENT_DATA_TYPE_FLOAT:
+			DataType = ENNETensorDataType::Float;
+			ElementSize = sizeof(float);
+			break;
+
+		case ONNX_TENSOR_ELEMENT_DATA_TYPE_UINT8:
+			DataType = ENNETensorDataType::UInt8;
+			ElementSize = sizeof(uint8);
+			break;
+		case ONNX_TENSOR_ELEMENT_DATA_TYPE_INT8:
+			DataType = ENNETensorDataType::Int8;
+			ElementSize = sizeof(int8);
+			break;
+		case ONNX_TENSOR_ELEMENT_DATA_TYPE_UINT16:
+			DataType = ENNETensorDataType::UInt16;
+			ElementSize = sizeof(uint16);
+			break;
+		case ONNX_TENSOR_ELEMENT_DATA_TYPE_INT16:
+			DataType = ENNETensorDataType::Int16;
+			ElementSize = sizeof(int16);
+			break;
+
+		case ONNX_TENSOR_ELEMENT_DATA_TYPE_INT32:
+			DataType = ENNETensorDataType::Int32;
+			ElementSize = sizeof(int32);
+			break;
+
+		case ONNX_TENSOR_ELEMENT_DATA_TYPE_INT64:
+			DataType = ENNETensorDataType::Int64;
+			ElementSize = sizeof(int64);
+			break;
+
+		case ONNX_TENSOR_ELEMENT_DATA_TYPE_STRING:
+			DataType = ENNETensorDataType::Char;
+			ElementSize = sizeof(char);
+			break;
+
+		case ONNX_TENSOR_ELEMENT_DATA_TYPE_BOOL:
+			DataType = ENNETensorDataType::Boolean;
+			ElementSize = sizeof(bool);
+			break;
+
+		case ONNX_TENSOR_ELEMENT_DATA_TYPE_FLOAT16:
+			DataType = ENNETensorDataType::Half;
+			ElementSize = 2;
+			break;
+
+		case ONNX_TENSOR_ELEMENT_DATA_TYPE_DOUBLE:
+			DataType = ENNETensorDataType::Double;
+			ElementSize = sizeof(double);
+			break;
+
+		case ONNX_TENSOR_ELEMENT_DATA_TYPE_UINT32:
+			DataType = ENNETensorDataType::UInt32;
+			ElementSize = sizeof(uint32);
+			break;
+
+		case ONNX_TENSOR_ELEMENT_DATA_TYPE_UINT64:
+			DataType = ENNETensorDataType::UInt64;
+			ElementSize = sizeof(uint64);
+			break;
+
+		case ONNX_TENSOR_ELEMENT_DATA_TYPE_COMPLEX64:
+			DataType = ENNETensorDataType::Complex64;
+			ElementSize = 8;
+			break;
+
+		case ONNX_TENSOR_ELEMENT_DATA_TYPE_COMPLEX128:
+			DataType = ENNETensorDataType::Complex128;
+			ElementSize = 16;
+			break;
+
+		case ONNX_TENSOR_ELEMENT_DATA_TYPE_BFLOAT16:
+			DataType = ENNETensorDataType::BFloat16;
+			ElementSize = 2;
+			break;
+
+		default:
+			DataType = ENNETensorDataType::None;
+			break;
+		}
+
+		return TypeInfoORT{ DataType, ElementSize };
+	}
+
+	inline void BindTensorsToORT(
+		TConstArrayView<NNECore::FTensorBindingCPU> InBindingTensors,
+		TConstArrayView<NNECore::Internal::FTensor> InTensors,
+		TConstArrayView<ONNXTensorElementDataType> InTensorsORTType,
+		const Ort::MemoryInfo* InAllocatorInfo,
+		TArray<Ort::Value>& OutOrtTensors
+	)
+	{
+		const uint32 NumBinding = InBindingTensors.Num();
+		const uint32 NumDescriptors = InTensors.Num();
+
+		if (NumBinding != NumDescriptors)
+		{
+			UE_LOG(LogNNE, Error, TEXT("BindTensorsToORT: Number of Bindings is different of Descriptors."));
+			return;
+		}
+
+		for (uint32 Index = 0; Index < NumBinding; ++Index)
+		{
+			const NNECore::FTensorBindingCPU& Binding = InBindingTensors[Index];
+			const NNECore::Internal::FTensor& Tensor = InTensors[Index];
+			const ONNXTensorElementDataType CurrentORTType = InTensorsORTType[Index];
+
+			TUniquePtr<int64_t[]> SizesInt64t;
+			SizesInt64t = MakeUnique<int64_t[]>(Tensor.GetShape().Rank());
+			for (int32 DimIndex = 0; DimIndex < Tensor.GetShape().Rank(); ++DimIndex)
+			{
+				SizesInt64t.Get()[DimIndex] = Tensor.GetShape().GetData()[DimIndex];
+			}
+
+			const uint64 ByteCount{ InTensors[Index].GetDataSize() };
+			const uint32 ArrayDimensions{ (uint32)Tensor.GetShape().Rank() };
+			OutOrtTensors.Emplace(
+				Ort::Value::CreateTensor(
+					*InAllocatorInfo,
+					Binding.Data,
+					ByteCount,
+					SizesInt64t.Get(),
+					ArrayDimensions,
+					CurrentORTType
+				)
+			);
+		}
+	}
+
+	inline void CopyFromORTToBindings(
+		TConstArrayView<Ort::Value> InOrtTensors,
+		TConstArrayView<NNECore::FTensorBindingCPU> InBindingTensors,
+		TConstArrayView<NNECore::FTensorDesc> InTensorDescs,
+		TArray<NNECore::Internal::FTensor>& OutTensors
+	)
+	{
+		const uint32 NumBinding = InOrtTensors.Num();
+		const uint32 NumDescriptors = InTensorDescs.Num();
+
+		if (NumBinding != NumDescriptors)
+		{
+			UE_LOG(LogNNE, Error, TEXT("CopyFromORTToBindings: Number of Bindings is different of Descriptors."));
+			return;
+		}
+
+		for (uint32 Index = 0; Index < NumBinding; ++Index)
+		{
+			const NNECore::FTensorBindingCPU& Binding = InBindingTensors[Index];
+			const NNECore::FTensorDesc& TensorDesc = InTensorDescs[Index];
+			const Ort::Value& OrtTensor = InOrtTensors[Index];
+			const std::vector<int64_t>& OrtShape = OrtTensor.GetTensorTypeAndShapeInfo().GetShape();
+
+			TArray<uint32> ShapeData;
+			for (int32 DimIndex = 0; DimIndex < OrtShape.size(); ++DimIndex)
+			{
+				check(OrtShape[DimIndex] >= 0);
+				ShapeData.Add(OrtShape[DimIndex]);
+			}
+
+			NNECore::FTensorShape Shape = NNECore::FTensorShape::Make(ShapeData);
+			NNECore::Internal::FTensor Tensor = NNECore::Internal::FTensor::Make(TensorDesc.GetName(), Shape, TensorDesc.GetDataType());
+			OutTensors.Emplace(Tensor);
+
+			const void* OrtTensorData = OrtTensor.GetTensorData<void>();
+			const uint64 DataSize = FMath::Min(Tensor.GetDataSize(), Binding.SizeInBytes);
+
+			if (DataSize > 0)
+			{
+				FMemory::Memcpy(Binding.Data, OrtTensorData, DataSize);
+			}
+
+			if (Tensor.GetDataSize() > Binding.SizeInBytes)
+			{
+				UE_LOG(LogNNE, Error, TEXT("CopyFromORTToBindings: Binding buffer was not large enough to contain all of the data, only first %d bytes were copied."), DataSize);
+			}
+		}
+	}
+
+} // UE::NNERuntimeCPU::Private
\ No newline at end of file
