NNE - Cleanup following feedback, clarify ownership of data between UNNEModelData and FModel

Problem:
* Previous CL regressed API extensibility in term of asset (UNNEModelData) storage optimisation.

Solution:
* UNNEModelData should always own the data and weights are neither shared  with FModels, if weight sharing is needed API should be extended at the FModel factory level allowing to specialize the sharing behavior for a given type of runtime interface (allowing for example dedicated hardware weight sharing). This simplify calling code and make the UNNEModelData asset lifetime easy to manage (no dependancy on it).

Test:
* All green test editor win64 + style transfer demo works

diff --git a/Engine/Plugins/Experimental/NNE/Source/NNEEditorTools/Private/NNEModelDataFactory.cpp b/Engine/Plugins/Experimental/NNE/Source/NNEEditorTools/Private/NNEModelDataFactory.cpp
--- a/Engine/Plugins/Experimental/NNE/Source/NNEEditorTools/Private/NNEModelDataFactory.cpp
+++ b/Engine/Plugins/Experimental/NNE/Source/NNEEditorTools/Private/NNEModelDataFactory.cpp
@@ -20,7 +20,8 @@ UObject* UNNEModelDataFactory::FactoryCreateBinary(UClass* Class, UObject* InPar
 	}
 
 	UNNEModelData* Result = NewObject<UNNEModelData>(InParent, Class, Name, Flags);
-	Result->Init(Type, Buffer, BufferEnd);
+	TConstArrayView<uint8> BufferView = MakeArrayView(Buffer, BufferEnd-Buffer);
+	Result->Init(Type, BufferView);
 
 	return Result;
 }
diff --git a/Engine/Plugins/Experimental/NNE/Source/NNXCore/Private/NNECoreModelData.cpp b/Engine/Plugins/Experimental/NNE/Source/NNXCore/Private/NNECoreModelData.cpp
--- a/Engine/Plugins/Experimental/NNE/Source/NNXCore/Private/NNECoreModelData.cpp
+++ b/Engine/Plugins/Experimental/NNE/Source/NNXCore/Private/NNECoreModelData.cpp
@@ -2,11 +2,13 @@
 
 #include "NNECoreModelData.h"
 
+#include "NNECore.h"
 #include "NNXCore.h"
 #include "NNXRuntimeFormat.h"
 #include "NNXModelOptimizerInterface.h"
 #include "NNECoreAttributeMap.h"
 #include "Serialization/CustomVersion.h"
+#include "UObject/WeakInterfacePtr.h"
 
 #if WITH_EDITOR
 #include "Misc/Guid.h"
@@ -59,33 +61,33 @@ inline void PutIntoDDC(const FGuid& FileDataId, const FString& RuntimeName, FSha
 
 inline TArray<uint8> Create(const FString& RuntimeName, FString FileType, const TArray<uint8>& FileData)
 {
-	using namespace NNX;
+	NNX::IRuntime* NNXRuntime = NNX::GetRuntime(RuntimeName);
+	if (NNXRuntime)
+	{
+		return NNXRuntime->CreateModelData(FileType, FileData);
+	}
 
-	// Try to optimize the model
-	IRuntime* Runtime = GetRuntime(RuntimeName);
-	if (!Runtime)
+	TWeakInterfacePtr<INNERuntime> NNERuntime = UE::NNECore::GetRuntime<INNERuntime>(RuntimeName);
+	if (NNERuntime.IsValid())
+	{
+		return NNERuntime->CreateModelData(FileType, FileData);
+	}
+	else
 	{
 		UE_LOG(LogNNX, Error, TEXT("UNNEModelData: No runtime '%s' found. Valid runtimes are: "), *RuntimeName);
-		TArray<NNX::IRuntime*> Runtimes = GetAllRuntimes();
+		TArrayView<TWeakInterfacePtr<INNERuntime>> Runtimes = UE::NNECore::GetAllRuntimes();
 		for (int i = 0; i < Runtimes.Num(); i++)
 		{
 			UE_LOG(LogNNX, Error, TEXT("- %s"), *Runtimes[i]->GetRuntimeName());
 		}
 		return {};
 	}
-
-	return Runtime->CreateModelData(FileType, FileData);
 }
 
-void UNNEModelData::Init(const TCHAR* Type, const uint8*& Buffer, const uint8* BufferEnd)
+void UNNEModelData::Init(const FString& Type, TConstArrayView<uint8> Buffer)
 {
-	check(Type)
-	check(Buffer)
-	check(BufferEnd)
-	check(BufferEnd - Buffer > 0)
-
-	FileType = FString(Type);
-	FileData = TArray<uint8>(Buffer, BufferEnd - Buffer);
+	FileType = Type;
+	FileData = Buffer;
 	FPlatformMisc::CreateGuid(FileDataId);
 	ModelData.Empty();
 }
diff --git a/Engine/Plugins/Experimental/NNE/Source/NNXCore/Public/NNECoreModelData.h b/Engine/Plugins/Experimental/NNE/Source/NNXCore/Public/NNECoreModelData.h
--- a/Engine/Plugins/Experimental/NNE/Source/NNXCore/Public/NNECoreModelData.h
+++ b/Engine/Plugins/Experimental/NNE/Source/NNXCore/Public/NNECoreModelData.h
@@ -14,11 +14,11 @@ class NNXCORE_API UNNEModelData : public UObject
 public:
 
 	/**
-	 * Function to initialize the data.
+	 * Function to initialize the data (will do a copy).
 	 * Called by the NNEModelDataFactory.
 	 * Type corresponds to the file extension.
 	 */
-	void Init(const TCHAR* Type, const uint8*& Buffer, const uint8* BufferEnd);
+	void Init(const FString& Type, TConstArrayView<uint8> Buffer);
 
 	/**
 	 * The function returns the cached (editor) or cooked (game) optimized model data for a given runtime.
@@ -31,8 +31,6 @@ public:
 
 	virtual void Serialize(FArchive& Ar) override;
 
-public:
-
 	/**
 	 * A GUID used for versioning.
 	 */
@@ -57,8 +55,6 @@ private:
 	 */
 	FGuid FileDataId;
 
-private:
-
 	/**
 	 * The processed / optimized model data for the different runtimes.
 	 */
diff --git a/Engine/Plugins/Experimental/NNE/Source/NNXCore/Public/NNECoreRuntimeCPU.h b/Engine/Plugins/Experimental/NNE/Source/NNXCore/Public/NNECoreRuntimeCPU.h
--- a/Engine/Plugins/Experimental/NNE/Source/NNXCore/Public/NNECoreRuntimeCPU.h
+++ b/Engine/Plugins/Experimental/NNE/Source/NNXCore/Public/NNECoreRuntimeCPU.h
@@ -59,6 +59,6 @@ class NNXCORE_API INNERuntimeCPU
 	GENERATED_BODY()
 	
 public:
-	virtual bool CanCreateModelCPU(TConstArrayView<uint8> Data) const = 0;
-	virtual TUniquePtr<UE::NNECore::IModelCPU> CreateModelCPU(TConstArrayView<uint8> Data) = 0;
+	virtual bool CanCreateModelCPU(TObjectPtr<UNNEModelData> ModelData) const = 0;
+	virtual TUniquePtr<UE::NNECore::IModelCPU> CreateModelCPU(TObjectPtr<UNNEModelData> ModelData) = 0;
 };
\ No newline at end of file
diff --git a/Engine/Plugins/Experimental/NNE/Source/NNXCore/Public/NNECoreRuntimeRDG.h b/Engine/Plugins/Experimental/NNE/Source/NNXCore/Public/NNECoreRuntimeRDG.h
--- a/Engine/Plugins/Experimental/NNE/Source/NNXCore/Public/NNECoreRuntimeRDG.h
+++ b/Engine/Plugins/Experimental/NNE/Source/NNXCore/Public/NNECoreRuntimeRDG.h
@@ -58,6 +58,6 @@ class NNXCORE_API INNERuntimeRDG
 	GENERATED_BODY()
 
 public:
-	virtual bool CanCreateModelRDG(TConstArrayView<uint8> Data) const = 0;
-	virtual TUniquePtr<UE::NNECore::IModelRDG> CreateModelRDG(TConstArrayView<uint8> Data) = 0;
+	virtual bool CanCreateModelRDG(TObjectPtr<UNNEModelData> ModelData) const = 0;
+	virtual TUniquePtr<UE::NNECore::IModelRDG> CreateModelRDG(TObjectPtr<UNNEModelData> ModelData) = 0;
 };
\ No newline at end of file
diff --git a/Engine/Plugins/Experimental/NNE/Source/NNXQA/Private/NNXQAUtils.cpp b/Engine/Plugins/Experimental/NNE/Source/NNXQA/Private/NNXQAUtils.cpp
--- a/Engine/Plugins/Experimental/NNE/Source/NNXQA/Private/NNXQAUtils.cpp
+++ b/Engine/Plugins/Experimental/NNE/Source/NNXQA/Private/NNXQAUtils.cpp
@@ -5,11 +5,13 @@
 #include "Kismet/GameplayStatics.h"
 #include "NNECore.h"
 #include "NNECoreAttributeMap.h"
+#include "NNECoreModelData.h"
 #include "NNECoreRuntime.h"
 #include "NNECoreRuntimeCPU.h"
 #include "NNXCore.h"
 #include "NNXInferenceModel.h"
 #include "NNXModelOptimizerInterface.h"
+#include "UObject/ObjectPtr.h"
 #include "UObject/WeakInterfacePtr.h"
 
 namespace NNX 
@@ -423,9 +425,9 @@ namespace Test
 			return TUniquePtr<UE::NNECore::IModelCPU>();
 		}
 
-		TArray<uint8> Data = RuntimeForModelData->CreateModelData(FString("onnx"), ONNXModelData.Data);
-		TUniquePtr<UE::NNECore::IModelCPU> InferenceModel = RuntimeForInferenceModel->CreateModelCPU(Data);
-		
+		TObjectPtr<UNNEModelData> ModelData = NewObject<UNNEModelData>();
+		ModelData->Init(FString("onnx"), ONNXModelData.Data);
+		TUniquePtr<UE::NNECore::IModelCPU> InferenceModel = RuntimeForInferenceModel->CreateModelCPU(ModelData);
 		return InferenceModel;
 	}
 	
diff --git a/Engine/Plugins/Experimental/NNE/Source/NNXRuntimeCPU/Private/NNERuntimeCPU.cpp b/Engine/Plugins/Experimental/NNE/Source/NNXRuntimeCPU/Private/NNERuntimeCPU.cpp
--- a/Engine/Plugins/Experimental/NNE/Source/NNXRuntimeCPU/Private/NNERuntimeCPU.cpp
+++ b/Engine/Plugins/Experimental/NNE/Source/NNXRuntimeCPU/Private/NNERuntimeCPU.cpp
@@ -49,10 +49,13 @@ TArray<uint8> UNNERuntimeCPUImpl::CreateModelData(FString FileType, TConstArrayV
 	return Result;
 }
 
-bool UNNERuntimeCPUImpl::CanCreateModelCPU(TConstArrayView<uint8> Data) const
+bool UNNERuntimeCPUImpl::CanCreateModelCPU(TObjectPtr<UNNEModelData> ModelData) const
 {
+	check(ModelData != nullptr);
+	
 	int32 GuidSize = sizeof(UNNERuntimeCPUImpl::GUID);
 	int32 VersionSize = sizeof(UNNERuntimeCPUImpl::Version);
+	TConstArrayView<uint8> Data = ModelData->GetModelData(GetRuntimeName());
 	
 	if (Data.Num() <= GuidSize + VersionSize)
 	{
@@ -64,15 +67,18 @@ bool UNNERuntimeCPUImpl::CanCreateModelCPU(TConstArrayView<uint8> Data) const
 	return bResult;
 }
 
-TUniquePtr<UE::NNECore::IModelCPU> UNNERuntimeCPUImpl::CreateModelCPU(TConstArrayView<uint8> Data)
+TUniquePtr<UE::NNECore::IModelCPU> UNNERuntimeCPUImpl::CreateModelCPU(TObjectPtr<UNNEModelData> ModelData)
 {
-	if (!CanCreateModelCPU(Data))
+	check(ModelData != nullptr);
+	
+	if (!CanCreateModelCPU(ModelData))
 	{
 		return TUniquePtr<UE::NNECore::IModelCPU>();
 	}
 
 	const UE::NNERuntimeCPU::Private::FRuntimeConf InConf;
 	UE::NNERuntimeCPU::Private::FModelCPU* Model = new UE::NNERuntimeCPU::Private::FModelCPU(&NNEEnvironmentCPU, InConf);
+	TConstArrayView<uint8> Data = ModelData->GetModelData(GetRuntimeName());
 	
 	if (!Model->Init(Data))
 	{
diff --git a/Engine/Plugins/Experimental/NNE/Source/NNXRuntimeCPU/Private/NNERuntimeCPU.h b/Engine/Plugins/Experimental/NNE/Source/NNXRuntimeCPU/Private/NNERuntimeCPU.h
--- a/Engine/Plugins/Experimental/NNE/Source/NNXRuntimeCPU/Private/NNERuntimeCPU.h
+++ b/Engine/Plugins/Experimental/NNE/Source/NNXRuntimeCPU/Private/NNERuntimeCPU.h
@@ -37,6 +37,6 @@ public:
 	virtual bool CanCreateModelData(FString FileType, TConstArrayView<uint8> FileData) const override;
 	virtual TArray<uint8> CreateModelData(FString FileType, TConstArrayView<uint8> FileData) override;
 
-	virtual bool CanCreateModelCPU(TConstArrayView<uint8> Data) const override;
-	virtual TUniquePtr<UE::NNECore::IModelCPU> CreateModelCPU(TConstArrayView<uint8> Data) override;
+	virtual bool CanCreateModelCPU(TObjectPtr<UNNEModelData> ModelData) const override;
+	virtual TUniquePtr<UE::NNECore::IModelCPU> CreateModelCPU(TObjectPtr<UNNEModelData> ModelData) override;
 };
\ No newline at end of file
