NNE - NNEHLslShader module API to internal + cleanup

Problem:
* NNEHLslShader module was exposing public API that we want to keep internal for 5.2

Solution:
* Move public headers to internal
* Also a bit of cleanup

Test:
* All tests green under win64 editor + style transfer in editor works

diff --git a/Engine/Plugins/Experimental/NNE/Source/NNEHlslShaders/Internal/NNXShaderParameters.h b/Engine/Plugins/Experimental/NNE/Source/NNEHlslShaders/Internal/NNXShaderParameters.h
new file mode 100644
--- /dev/null
+++ b/Engine/Plugins/Experimental/NNE/Source/NNEHlslShaders/Internal/NNXShaderParameters.h
@@ -0,0 +1,174 @@
+// Copyright Epic Games, Inc. All Rights Reserved.
+
+#pragma once
+
+#include "ShaderParameterUtils.h"
+#include "RenderGraphUtils.h"
+
+// HLSL runtime and DirectML runtime use different shader parameter access patterns
+#define NNXRT_RDG_DML
+
+#ifdef NNXRT_RDG_DML
+#define NNXRT_RDG_BUFFER_SRV(View, BufferName) RDG_BUFFER_ACCESS(BufferName, ERHIAccess::UAVCompute)
+#define NNXRT_RDG_BUFFER_UAV(View, BufferName) RDG_BUFFER_ACCESS(BufferName, ERHIAccess::UAVCompute)
+#else
+#define NNXRT_RDG_BUFFER_SRV(View, BufferName) SHADER_PARAMETER_RDG_BUFFER_SRV(View, BufferName)
+#define NNXRT_RDG_BUFFER_UAV(View, BufferName) SHADER_PARAMETER_RDG_BUFFER_UAV(View, BufferName)
+#endif
+
+#ifdef NNXRT_RDG_DML
+#define NNXRT_BEGIN_SHADER_PARAMETER_STRUCT(Name) BEGIN_SHADER_PARAMETER_STRUCT(Name, )
+#else
+#define NNXRT_BEGIN_SHADER_PARAMETER_STRUCT(Name) BEGIN_SHADER_PARAMETER_STRUCT(FParameters, )
+#endif
+
+#define NNXRT_END_SHADER_PARAMETER_STRUCT END_SHADER_PARAMETER_STRUCT
+
+#define NXRT_TENSORSTRIDEINFO_MAX_NUM_DIMENSIONS 8
+
+//Unary element wise operators
+#define NNXRT_ELEMENTWISEUNARY_PARAMETER_STRUCT() \
+	NNXRT_BEGIN_SHADER_PARAMETER_STRUCT(FNNEElementWiseUnaryParameters) \
+		NNXRT_RDG_BUFFER_SRV(Buffer<float>, Input) \
+		NNXRT_RDG_BUFFER_UAV(RWBuffer<float>, Output) \
+		SHADER_PARAMETER(uint32, Num) \
+		SHADER_PARAMETER(uint32, ThreadCountX) \
+		SHADER_PARAMETER(float, Alpha) \
+		SHADER_PARAMETER(float, Beta) \
+		SHADER_PARAMETER(float, Gamma) \
+	NNXRT_END_SHADER_PARAMETER_STRUCT()
+
+#ifdef NNXRT_RDG_DML
+NNXRT_ELEMENTWISEUNARY_PARAMETER_STRUCT()
+#endif
+
+//Variadic element wise operators
+#define NNXRT_ELEMENTWISEVARIADIC_PARAMETER_STRUCT() \
+	NNXRT_BEGIN_SHADER_PARAMETER_STRUCT(FNNEElementWiseVariadicParameters) \
+		NNXRT_RDG_BUFFER_SRV(Buffer<float>, Input0) \
+		NNXRT_RDG_BUFFER_SRV(Buffer<float>, Input1) \
+		NNXRT_RDG_BUFFER_SRV(Buffer<float>, Input2) \
+		NNXRT_RDG_BUFFER_SRV(Buffer<float>, Input3) \
+		NNXRT_RDG_BUFFER_UAV(RWBuffer<float>, Output) \
+		SHADER_PARAMETER_ARRAY(FUintVector4, InputTensorInfo, [NXRT_TENSORSTRIDEINFO_MAX_NUM_DIMENSIONS]) \
+		SHADER_PARAMETER_ARRAY(FUintVector4, OutputTensorInfo, [NXRT_TENSORSTRIDEINFO_MAX_NUM_DIMENSIONS]) \
+		SHADER_PARAMETER(uint32, Num) \
+		SHADER_PARAMETER(uint32, ThreadCountX) \
+		SHADER_PARAMETER(float, Scale) \
+	NNXRT_END_SHADER_PARAMETER_STRUCT()
+
+#ifdef NNXRT_RDG_DML
+NNXRT_ELEMENTWISEVARIADIC_PARAMETER_STRUCT()
+#endif
+
+//Binary element wise operators
+#define NNXRT_ELEMENTWISEBINARY_PARAMETER_STRUCT() \
+	NNXRT_BEGIN_SHADER_PARAMETER_STRUCT(FNNEElementWiseBinaryParameters) \
+		NNXRT_RDG_BUFFER_SRV(Buffer<float>, LHSInput) \
+		NNXRT_RDG_BUFFER_SRV(Buffer<float>, RHSInput) \
+		NNXRT_RDG_BUFFER_UAV(RWBuffer<float>, Output) \
+		SHADER_PARAMETER_ARRAY(FUintVector4, TensorInfo, [NXRT_TENSORSTRIDEINFO_MAX_NUM_DIMENSIONS]) \
+		SHADER_PARAMETER(uint32, Num) \
+		SHADER_PARAMETER(uint32, ThreadCountX) \
+	NNXRT_END_SHADER_PARAMETER_STRUCT()
+
+#ifdef NNXRT_RDG_DML
+NNXRT_ELEMENTWISEBINARY_PARAMETER_STRUCT()
+#endif
+
+//Gemm
+#define NNXRT_GEMM_MAX_NUM_STACK_DIMENSIONS 8
+
+#define NNXRT_GEMM_PARAMETER_STRUCT() \
+	NNXRT_BEGIN_SHADER_PARAMETER_STRUCT(FNNEGemmParameters) \
+		SHADER_PARAMETER(float, Alpha) \
+		SHADER_PARAMETER(float, Beta) \
+		SHADER_PARAMETER(int32, TransA) \
+		SHADER_PARAMETER(int32, TransB) \
+		SHADER_PARAMETER(uint32, M) \
+		SHADER_PARAMETER(uint32, N) \
+		SHADER_PARAMETER(uint32, K) \
+		SHADER_PARAMETER(uint32, MxK) \
+		SHADER_PARAMETER(uint32, KxN) \
+		SHADER_PARAMETER(uint32, MxN) \
+		SHADER_PARAMETER(uint32, CWidth) \
+		SHADER_PARAMETER(uint32, CHeight) \
+		SHADER_PARAMETER(float, CScalar) \
+		SHADER_PARAMETER_ARRAY(FUint32Vector4, StackShapeA_StackShapeB_StackStrideA_StackStrideB, [NNXRT_GEMM_MAX_NUM_STACK_DIMENSIONS]) \
+		NNXRT_RDG_BUFFER_SRV(Buffer<float>, A) \
+		NNXRT_RDG_BUFFER_SRV(Buffer<float>, B) \
+		NNXRT_RDG_BUFFER_SRV(Buffer<float>, C) \
+		NNXRT_RDG_BUFFER_UAV(RWBuffer<float>, Y) \
+	NNXRT_END_SHADER_PARAMETER_STRUCT()
+
+#ifdef NNXRT_RDG_DML
+NNXRT_GEMM_PARAMETER_STRUCT()
+#endif
+
+//Conv
+#define NNXRT_CONV_MAX_NUM_DIMENSIONS 4
+
+#define NNXRT_CONV_PARAMETER_STRUCT() \
+	NNXRT_BEGIN_SHADER_PARAMETER_STRUCT(FNNEConvParameters) \
+		NNXRT_RDG_BUFFER_SRV(Buffer<float>, X) \
+		NNXRT_RDG_BUFFER_SRV(Buffer<float>, W) \
+		NNXRT_RDG_BUFFER_UAV(RWBuffer<float>, Y) \
+		NNXRT_RDG_BUFFER_SRV(Buffer<float>, B) \
+		SHADER_PARAMETER_ARRAY(FIntVector4, Dilation_Stride_XBlockStartOffset_DilationXBlockStride, [NNXRT_CONV_MAX_NUM_DIMENSIONS]) \
+		SHADER_PARAMETER_ARRAY(FIntVector4, GroupStride_GroupShape_GroupThreadStride_StrideXBlockStride, [NNXRT_CONV_MAX_NUM_DIMENSIONS]) \
+		SHADER_PARAMETER_ARRAY(FIntVector4, YDimension_YMemoryStride_XDimension_XMemoryStride, [NNXRT_CONV_MAX_NUM_DIMENSIONS]) \
+		SHADER_PARAMETER_ARRAY(FIntVector4, XBlockStartStride_XBlockStride_WDimension_WDimensionDilationXBlockStride, [NNXRT_CONV_MAX_NUM_DIMENSIONS]) \
+		SHADER_PARAMETER_ARRAY(FVector4f, OneDiv_GroupStride_GroupThreadStride_XBlockStride, [NNXRT_CONV_MAX_NUM_DIMENSIONS]) \
+		SHADER_PARAMETER(int32, NumWChannels) \
+		SHADER_PARAMETER(int32, YBatchStride) \
+		SHADER_PARAMETER(int32, YOutputKernelStride) \
+		SHADER_PARAMETER(int32, XBatchStride) \
+		SHADER_PARAMETER(int32, XChannelStride) \
+		SHADER_PARAMETER(int32, XBlockSize) \
+		SHADER_PARAMETER(int32, NumChannelBatches) \
+		SHADER_PARAMETER(int32, NumChannelsPerBatch) \
+		SHADER_PARAMETER(int32, WOutputKernelStride) \
+		SHADER_PARAMETER(int32, WChannelBatchSize) \
+		SHADER_PARAMETER(int32, WChannelSize) \
+		SHADER_PARAMETER(float, GroupsDivM) \
+	NNXRT_END_SHADER_PARAMETER_STRUCT()
+
+#ifdef NNXRT_RDG_DML
+NNXRT_CONV_PARAMETER_STRUCT()
+#endif
+
+//ConvTranspose
+#define NNXRT_CONVTRANSPOSE_MAX_NUM_STACK_DIMENSIONS 4
+#define NNXRT_CONVTRANSPOSE_MIN_NUM_READS_PER_THREAD_POW2 1
+#define NNXRT_CONVTRANSPOSE_MAX_NUM_READS_PER_THREAD_POW2 3
+
+#define NNXRT_CONVTRANSPOSE_PARAMETER_STRUCT() \
+	NNXRT_BEGIN_SHADER_PARAMETER_STRUCT(FNNEConvTransposeParameters) \
+		NNXRT_RDG_BUFFER_SRV(Buffer<float>, X) \
+		NNXRT_RDG_BUFFER_SRV(Buffer<float>, W) \
+		NNXRT_RDG_BUFFER_UAV(RWBuffer<float>, Y) \
+		NNXRT_RDG_BUFFER_SRV(Buffer<float>, B) \
+		SHADER_PARAMETER_ARRAY(FIntVector4, Dilation_Stride_XBlockStartOffset_DilationXBlockStride, [NNXRT_CONVTRANSPOSE_MAX_NUM_STACK_DIMENSIONS]) \
+		SHADER_PARAMETER_ARRAY(FIntVector4, GroupStride_GroupShape_GroupThreadStride_StrideXBlockStride, [NNXRT_CONVTRANSPOSE_MAX_NUM_STACK_DIMENSIONS]) \
+		SHADER_PARAMETER_ARRAY(FIntVector4, YDimension_YMemoryStride_XDimension_XMemoryStride, [NNXRT_CONVTRANSPOSE_MAX_NUM_STACK_DIMENSIONS]) \
+		SHADER_PARAMETER_ARRAY(FIntVector4, XBlockStartStride_XBlockStride_WDimension_WDimensionDilationXBlockStride, [NNXRT_CONVTRANSPOSE_MAX_NUM_STACK_DIMENSIONS]) \
+		SHADER_PARAMETER_ARRAY(FVector4f, OneDiv_GroupStride_GroupThreadStride_OneDivStride, [NNXRT_CONVTRANSPOSE_MAX_NUM_STACK_DIMENSIONS]) \
+		SHADER_PARAMETER(int32, NumWChannels) \
+		SHADER_PARAMETER(int32, NumOutChannelsDivGroup) \
+		SHADER_PARAMETER(int32, YBatchStride) \
+		SHADER_PARAMETER(int32, YOutputKernelStride) \
+		SHADER_PARAMETER(int32, XBatchStride) \
+		SHADER_PARAMETER(int32, XChannelStride) \
+		SHADER_PARAMETER(int32, XBlockSize) \
+		SHADER_PARAMETER(int32, NumChannelBatches) \
+		SHADER_PARAMETER(int32, NumChannelsPerBatch) \
+		SHADER_PARAMETER(int32, WOutputKernelStride) \
+		SHADER_PARAMETER(int32, WChannelBatchSize) \
+		SHADER_PARAMETER(int32, WChannelSize) \
+		SHADER_PARAMETER(float, GroupsDivM) \
+		SHADER_PARAMETER(float, OneDivGroup) \
+	NNXRT_END_SHADER_PARAMETER_STRUCT()
+
+#ifdef NNXRT_RDG_DML
+NNXRT_CONVTRANSPOSE_PARAMETER_STRUCT()
+#endif
\ No newline at end of file
diff --git a/Engine/Plugins/Experimental/NNE/Source/NNEHlslShaders/NNEHlslShaders.Build.cs b/Engine/Plugins/Experimental/NNE/Source/NNEHlslShaders/NNEHlslShaders.Build.cs
--- a/Engine/Plugins/Experimental/NNE/Source/NNEHlslShaders/NNEHlslShaders.Build.cs
+++ b/Engine/Plugins/Experimental/NNE/Source/NNEHlslShaders/NNEHlslShaders.Build.cs
@@ -8,16 +8,10 @@ public class NNEHlslShaders : ModuleRules
 	{
         PCHUsage = PCHUsageMode.UseExplicitOrSharedPCHs;
 
-		PublicDependencyModuleNames.AddRange(
-			new string[] {
-				"Core",
-				"CoreUObject",
-				"Engine"
-			}
-		);
-
 		PrivateDependencyModuleNames.AddRange(
 			new string[] {
+				"Core",
+				"Engine",
 				"Projects",
 				"RenderCore",
 				"RHI",
diff --git a/Engine/Plugins/Experimental/NNE/Source/NNEHlslShaders/Private/NNEHlslShadersConvCS.cpp b/Engine/Plugins/Experimental/NNE/Source/NNEHlslShaders/Private/NNEHlslShadersConvCS.cpp
--- a/Engine/Plugins/Experimental/NNE/Source/NNEHlslShaders/Private/NNEHlslShadersConvCS.cpp
+++ b/Engine/Plugins/Experimental/NNE/Source/NNEHlslShaders/Private/NNEHlslShadersConvCS.cpp
@@ -5,7 +5,7 @@
 
 namespace UE::NNEHlslShaders::Internal
 {
-	namespace
+	namespace ConvUtils
 	{
 
 	TArray<int32> GetXBlockShape(TArrayView<const int32> GroupShape, TArrayView<const uint32> WShape, TArrayView<const int32> Dilations, TArrayView<const int32> Strides)
@@ -101,7 +101,7 @@ namespace UE::NNEHlslShaders::Internal
 		return Result;
 	}
 
-	} // namespace
+	} // namespace ConvUtils
 
 	void FConvCS::ModifyCompilationEnvironment(const FGlobalShaderPermutationParameters& InParameters, FShaderCompilerEnvironment& OutEnvironment)
 	{
@@ -117,7 +117,7 @@ namespace UE::NNEHlslShaders::Internal
 		check(Strides.Num() == 0 || Strides.Num() == WShape.Num() - 2);
 		check(Pads.Num() == 0 || Pads.Num() == 2 * (WShape.Num() - 2));
 
-		TArray<int32> Padding = GetPadding(XShape, WShape, AutoPad, Dilations, Strides, Pads);
+		TArray<int32> Padding = ConvUtils::GetPadding(XShape, WShape, AutoPad, Dilations, Strides, Pads);
 
 		TArray<int32> Result;
 		Result.SetNumUninitialized(XShape.Num());
@@ -151,11 +151,11 @@ namespace UE::NNEHlslShaders::Internal
 
 		int32 NumDimensions = XShape.Num() - 2;
 
-		TArray<int32> Padding = GetPadding(XShape, WShape, AutoPad, Dilations, Strides, Pads);
+		TArray<int32> Padding = ConvUtils::GetPadding(XShape, WShape, AutoPad, Dilations, Strides, Pads);
 		TArray<int32> GroupShape = GetGroupShape(GroupSize, NumDimensions);
 		TArray<int32> YShape = GetOutputShape(XShape, WShape, AutoPad, Dilations, Strides, Pads);
-		TArray<int32> GridShape = GetGridShape(YShape, GroupShape);
-		TArray<int32> XBlockShape = GetXBlockShape(GroupShape, WShape, Dilations, Strides);
+		TArray<int32> GridShape = ConvUtils::GetGridShape(YShape, GroupShape);
+		TArray<int32> XBlockShape = ConvUtils::GetXBlockShape(GroupShape, WShape, Dilations, Strides);
 		
 		int32 GroupStride = 1;
 		int32 GroupThreadStride = 1;
@@ -217,7 +217,7 @@ namespace UE::NNEHlslShaders::Internal
 			NumThreadsPerGroup *= GroupShape[i];
 		}
 
-		TArray<int32> XBlockShape = GetXBlockShape(GroupShape, WShape, Dilations, Strides);
+		TArray<int32> XBlockShape = ConvUtils::GetXBlockShape(GroupShape, WShape, Dilations, Strides);
 		int32 NumXBlockElements = 1;
 		for (int32 i = 0; i < NumDimensions; i++)
 		{
@@ -239,7 +239,7 @@ namespace UE::NNEHlslShaders::Internal
 	{
 		check(NumDimensions > 0);
 
-		int32 NumThreadsPerGroup = GetNumThreadsPerGroup(GroupSize);
+		int32 NumThreadsPerGroup = ConvUtils::GetNumThreadsPerGroup(GroupSize);
 
 		int32 Power = (int32)FMath::Log2((float)NumThreadsPerGroup);
 		int32 MinPowerPerDim = (int32)((float)Power / (float)NumDimensions);
@@ -279,7 +279,7 @@ namespace UE::NNEHlslShaders::Internal
 		
 		for (int32 i = 0; i < (int32)EConvGroupSize::MAX; i++)
 		{
-			if (GetNumThreadsPerGroup((EConvGroupSize)i) >= WChannelSize) 
+			if (ConvUtils::GetNumThreadsPerGroup((EConvGroupSize)i) >= WChannelSize)
 			{
 				return (EConvGroupSize)i;
 			}
diff --git a/Engine/Plugins/Experimental/NNE/Source/NNEHlslShaders/Private/NNEHlslShadersElementWiseUnaryCS.cpp b/Engine/Plugins/Experimental/NNE/Source/NNEHlslShaders/Private/NNEHlslShadersElementWiseUnaryCS.cpp
--- a/Engine/Plugins/Experimental/NNE/Source/NNEHlslShaders/Private/NNEHlslShadersElementWiseUnaryCS.cpp
+++ b/Engine/Plugins/Experimental/NNE/Source/NNEHlslShaders/Private/NNEHlslShadersElementWiseUnaryCS.cpp
@@ -46,7 +46,7 @@ namespace UE::NNEHlslShaders::Internal
 		OP(Exp,         TEXT("exp(X)"));
 		OP(Floor,       TEXT("floor(X)"));
 		OP(IsInf,       TEXT("isinf(X)"));
-		OP(IsNan,       TEXT("isnan(X)"));//TODO check shader there is a warning saying input can neither be Nan on PC FXC
+		OP(IsNan,       TEXT("isnan(X)"));//Note: There is a warning on PC FXC about input that can neither be Nan.
 		OP(HardSigmoid, TEXT("hardSigmoid(X)"));
 		OP(HardSwish,   TEXT("hardSwish(X)"));
 		OP(LeakyRelu,   TEXT("leakyRelu(X)"));
diff --git a/Engine/Plugins/Experimental/NNE/Source/NNEHlslShaders/Private/NNEHlslShadersInstanceNormalizationCS.cpp b/Engine/Plugins/Experimental/NNE/Source/NNEHlslShaders/Private/NNEHlslShadersInstanceNormalizationCS.cpp
--- a/Engine/Plugins/Experimental/NNE/Source/NNEHlslShaders/Private/NNEHlslShadersInstanceNormalizationCS.cpp
+++ b/Engine/Plugins/Experimental/NNE/Source/NNEHlslShaders/Private/NNEHlslShadersInstanceNormalizationCS.cpp
@@ -6,7 +6,7 @@
 
 namespace UE::NNEHlslShaders::Internal
 {
-	namespace
+	namespace InstNormUtils
 	{
 		FIntVector GetGroupSize(EInstanceNormalizationAlgorithm Algorithm)
 		{
@@ -36,7 +36,7 @@ namespace UE::NNEHlslShaders::Internal
 
 	FIntVector TInstanceNormalizationCS::GetGroupCount(const TInstanceNormalizationCS::FParameters& Parameters, EInstanceNormalizationAlgorithm Algorithm)
 	{
-		int ThreadGroupCountValueY = FMath::DivideAndRoundUp((int32)Parameters.NxC, GetGroupSize(Algorithm).Y);
+		int ThreadGroupCountValueY = FMath::DivideAndRoundUp((int32)Parameters.NxC, InstNormUtils::GetGroupSize(Algorithm).Y);
 
 		return {1, ThreadGroupCountValueY, 1};
 	}
diff --git a/Engine/Plugins/Experimental/NNE/Source/NNEHlslShaders/Public/NNXShaderParameters.h b/Engine/Plugins/Experimental/NNE/Source/NNEHlslShaders/Public/NNXShaderParameters.h
deleted file mode 100644
