NNE - NNEUtils module API to internal + cleanup

Problem:
* NNEUtils module was exposing public API that we want to keep internal for 5.2

Solution:
* Move public headers to internal
* Also update module to follow NNE naming convention + renamed module to NNEUtils (from NNXUtils)

Test:
* All tests green under win64 editor + style transfer in editor works

diff --git a/Engine/Plugins/Experimental/NNE/NNE.uplugin b/Engine/Plugins/Experimental/NNE/NNE.uplugin
--- a/Engine/Plugins/Experimental/NNE/NNE.uplugin
+++ b/Engine/Plugins/Experimental/NNE/NNE.uplugin
@@ -32,7 +32,7 @@
 			"PlatformAllowList": [ "Win64", "Linux", "Mac" ]
 		},
 		{
-			"Name": "NNXUtils",
+			"Name": "NNEUtils",
 			"Type": "RuntimeAndProgram",
 			"LoadingPhase": "Default",
 			"PlatformAllowList": [ "Win64", "Linux", "Mac" ]
diff --git a/Engine/Plugins/Experimental/NNE/Source/NNEUtils/Internal/NNEUtilsModelBuilder.h b/Engine/Plugins/Experimental/NNE/Source/NNEUtils/Internal/NNEUtilsModelBuilder.h
new file mode 100644
--- /dev/null
+++ b/Engine/Plugins/Experimental/NNE/Source/NNEUtils/Internal/NNEUtilsModelBuilder.h
@@ -0,0 +1,103 @@
+// Copyright Epic Games, Inc. All Rights Reserved.
+#pragma once
+
+#include "NNECoreTypes.h"
+#include "NNECoreTensor.h"
+
+namespace UE::NNECore { class FAttributeMap; }
+struct FNNEAttributeValue;
+struct FNNIModelRaw;
+
+namespace UE::NNEUtils::Internal
+{
+using FTensor = NNECore::Internal::FTensor;
+
+class IModelBuilder
+{
+public:
+
+	enum class HandleType : uint8
+	{
+		Invalid,
+		Tensor,
+		Operator
+	};
+
+	template<typename Tag>
+	struct Handle
+	{
+		void*		Ptr { nullptr };
+		HandleType	Type { HandleType::Invalid };
+	};
+
+	typedef struct Handle<struct TensorTag>		HTensor;
+	typedef struct Handle<struct OperatorTag>	HOperator;
+
+	template<typename TensorT>
+	HTensor MakeTensorHandle(TensorT* TensorPtr)
+	{
+		return HTensor{ TensorPtr, HandleType::Tensor };
+	}
+
+	template<typename OperatorT>
+	HOperator MakeOperatorHandle(OperatorT* OperatorPtr)
+	{
+		return HOperator{ OperatorPtr, HandleType::Operator };
+	}
+
+	virtual ~IModelBuilder() = default;
+
+	/** Initialize the model builder */
+	virtual bool Begin(const FString& Name = TEXT("main")) = 0;
+
+	/** Serialize the model to given array */
+	virtual bool End(TArray<uint8>& Data) = 0;
+
+	/** Add tensor */
+	virtual HTensor AddTensor(const FString& Name, ENNETensorDataType DataType, TArrayView<const int32> Shape, const void* Data = nullptr, uint64 DataSize = 0) = 0;
+
+	/** Add model input */
+	virtual bool AddInput(HTensor InTensor) = 0;
+
+	/** Add model output */
+	virtual bool AddOutput(HTensor OutTensor) = 0;
+
+	/** Add operator */
+	virtual HOperator AddOperator(const FString& Type, const FString& Name = TEXT("")) = 0;
+
+	/** Add operator input */
+	virtual bool AddOperatorInput(HOperator Op, HTensor Tensor) = 0;
+
+	/** Add operator attribute */
+	virtual bool AddOperatorAttribute(HOperator Op, const FString& Name, const FNNEAttributeValue& Value) = 0;
+	
+	/** Add operator output */
+	virtual bool AddOperatorOutput(HOperator Op, HTensor Tensor) = 0;
+};
+
+/** Default ONNX IR version */
+static constexpr int64 OnnxIrVersion = 7;
+
+/** Default ONNX operator set version */
+static constexpr int64 OnnxOpsetVersion = 15;
+
+/**
+ * Create an instance of ONNX model builder that creates ONNX models in memory
+ */
+NNEUTILS_API IModelBuilder* CreateONNXModelBuilder(int64 IrVersion = OnnxIrVersion, int64 OpsetVersion = OnnxOpsetVersion);
+
+/**
+ * Utility functions to create single layer NN for operator testing with optional attributes
+ */
+NNEUTILS_API bool CreateONNXModelForOperator(bool UseVariadicShapeForModel, const FString& OperatorName, 
+	TConstArrayView<FTensor> InInputTensors, TConstArrayView<FTensor> InOutputTensors,
+	TConstArrayView<FTensor> InWeightTensors, TConstArrayView<TArray<char>> InWeightTensorsData,
+	const UE::NNECore::FAttributeMap& Attributes, FNNIModelRaw& ModelData);
+
+/**
+ * Create an instance of NNX model builder that creates NNX model/format in memory
+ */
+NNEUTILS_API IModelBuilder* CreateNNEModelBuilder();
+
+} // UE::NNEUtils::Internal
+
diff --git a/Engine/Plugins/Experimental/NNE/Source/NNEUtils/Internal/NNEUtilsModelOptimizer.h b/Engine/Plugins/Experimental/NNE/Source/NNEUtils/Internal/NNEUtilsModelOptimizer.h
new file mode 100644
--- /dev/null
+++ b/Engine/Plugins/Experimental/NNE/Source/NNEUtils/Internal/NNEUtilsModelOptimizer.h
@@ -0,0 +1,32 @@
+// Copyright Epic Games, Inc. All Rights Reserved.
+#pragma once
+
+#include "NNXRuntimeFormat.h"
+
+namespace UE::NNECore { class FAttributeMap; }
+namespace NNX { class IModelOptimizer; }
+
+namespace UE::NNEUtils::Internal
+{
+
+using FOptimizerOptionsMap = NNECore::FAttributeMap;
+	
+/** Create a model optimizer */
+NNEUTILS_API TUniquePtr<NNX::IModelOptimizer> CreateModelOptimizer(ENNXInferenceFormat InputFormat, ENNXInferenceFormat OutputFormat);
+
+inline TUniquePtr<NNX::IModelOptimizer> CreateONNXToNNEModelOptimizer()
+{
+	return CreateModelOptimizer(ENNXInferenceFormat::ONNX, ENNXInferenceFormat::NNXRT);
+}
+
+inline TUniquePtr<NNX::IModelOptimizer> CreateONNXToORTModelOptimizer()
+{
+	return CreateModelOptimizer(ENNXInferenceFormat::ONNX, ENNXInferenceFormat::ORT);
+}
+
+inline TUniquePtr<NNX::IModelOptimizer> CreateONNXToONNXModelOptimizer()
+{
+	return CreateModelOptimizer(ENNXInferenceFormat::ONNX, ENNXInferenceFormat::ONNX);
+}
+
+} // UE::NNEUtils::Internal
diff --git a/Engine/Plugins/Experimental/NNE/Source/NNEUtils/NNEUtils.Build.cs b/Engine/Plugins/Experimental/NNE/Source/NNEUtils/NNEUtils.Build.cs
new file mode 100644
--- /dev/null
+++ b/Engine/Plugins/Experimental/NNE/Source/NNEUtils/NNEUtils.Build.cs
@@ -0,0 +1,55 @@
+// Copyright Epic Games, Inc. All Rights Reserved.
+
+
+using UnrealBuildTool;
+using System.IO;
+
+public class NNEUtils : ModuleRules
+{
+	public NNEUtils(ReadOnlyTargetRules Target) : base(Target)
+	{
+		ShortName = "NNEUtils"; // Shorten to avoid path-too-long errors
+		PCHUsage = ModuleRules.PCHUsageMode.UseExplicitOrSharedPCHs;
+
+		PublicIncludePaths.AddRange(
+			new string[] {
+				Path.GetFullPath(Path.Combine(EngineDirectory, "Plugins/Experimental/NNE/Source/NNXCore/Private")),
+			}
+		);
+
+		PublicDependencyModuleNames.AddRange(
+			new string[]
+			{
+				"NNXCore"
+			}
+		);
+
+		PrivateDependencyModuleNames.AddRange(
+			new string[]
+			{
+				"Core",
+				"CoreUObject",
+				"Engine",
+				"NNX_ONNXRuntime",
+				"NNX_ONNX_1_11_0",
+				"NNX_ONNXRuntimeProto_1_11_0",
+				"ORTHelper",
+				}
+			);
+
+
+		if (Target.Platform == UnrealTargetPlatform.Win64 || 
+			Target.Platform == UnrealTargetPlatform.Linux ||
+			Target.Platform == UnrealTargetPlatform.Mac)
+		{
+			PrivateDependencyModuleNames.AddRange
+				(
+				new string[] {
+					"Protobuf",
+					"Re2" // ONNXRuntimeRE2
+				}
+			);
+		}
+	}
+}
+
diff --git a/Engine/Plugins/Experimental/NNE/Source/NNEUtils/Private/NNEUtilsModelBuilderNNE.cpp b/Engine/Plugins/Experimental/NNE/Source/NNEUtils/Private/NNEUtilsModelBuilderNNE.cpp
new file mode 100644
--- /dev/null
+++ b/Engine/Plugins/Experimental/NNE/Source/NNEUtils/Private/NNEUtilsModelBuilderNNE.cpp
@@ -0,0 +1,294 @@
+// Copyright Epic Games, Inc. All Rights Reserved.
+
+#include "NNEUtilsModelBuilder.h"
+#include "NNECore.h"
+#include "NNECoreAttributeMap.h"
+#include "NNXRuntimeFormat.h"
+
+#include "Misc/StringBuilder.h"
+#include "Serialization/MemoryWriter.h"
+
+#define Print(Format, ...) UE_LOG(LogNNE, Display, Format, __VA_ARGS__)
+
+namespace UE::NNEUtils::Internal
+{
+
+class FModelPrinterNNE
+{
+public:
+
+	void Visit(const FMLRuntimeFormat& Format)
+	{
+		for (int Idx = 0; Idx < Format.Tensors.Num(); ++Idx)
+		{
+			Visit(Format.Tensors[Idx]);
+		}
+
+		for (int Idx = 0; Idx < Format.Operators.Num(); ++Idx)
+		{
+			Visit(Format.Operators[Idx]);
+		}
+	}
+
+	void Visit(const FMLFormatTensorDesc& Tensor)
+	{
+		FStringBuilderBase Str;
+
+		Str << TEXT("[");
+		
+		for (int32 Idx = 0; Idx < Tensor.Shape.Num(); ++Idx)
+		{
+			Str << Tensor.Shape[Idx];
+
+			if (Idx + 1 < Tensor.Shape.Num())
+			{
+				Str << TEXT(",");
+			}
+		}
+		
+		Str << TEXT("]");
+
+		Print(TEXT("Tensor:%s %s"), *Tensor.Name, Str.ToString());
+		
+	}
+
+	void Visit(const FMLFormatOperatorDesc& Op)
+	{
+		Print(TEXT("Op:%s in:%d out:%d"), *Op.TypeName, Op.InTensors.Num(), Op.OutTensors.Num());
+	}
+};
+
+inline int NNETensorCast(IModelBuilder::HTensor& Handle)
+{
+	if (Handle.Type == IModelBuilder::HandleType::Tensor)
+	{
+		return int(reinterpret_cast<int64>(Handle.Ptr));
+	}
+
+	return -1;
+}
+
+inline int NNEOperatorCast(IModelBuilder::HOperator& Handle)
+{
+	if (Handle.Type == IModelBuilder::HandleType::Operator)
+	{
+		return int(reinterpret_cast<int64>(Handle.Ptr));
+	}
+
+	return -1;
+}
+
+/**
+ * NNE format builder, create NNE format in memory
+ */
+class FModelBuilderNNE : public IModelBuilder
+{
+public:
+
+	FModelBuilderNNE()
+	{
+
+	}
+
+	virtual bool Begin(const FString& Name) override
+	{
+		return true;
+	}
+
+	virtual bool End(TArray<uint8>& Data) override
+	{
+		// This is for debugging purposes
+		FModelPrinterNNE Printer;
+
+		Printer.Visit(Format);
+
+		FMemoryWriter Writer(Data);
+
+		FMLRuntimeFormat::StaticStruct()->SerializeBin(Writer, &Format);
+
+		return !Data.IsEmpty();
+	}
+
+	virtual HTensor AddTensor(const FString& Name, ENNETensorDataType DataType, TArrayView<const int32> Shape, const void* Data, uint64 DataSize)
+	{
+		TArray<int32, TInlineAllocator<NNECore::FTensorShape::MaxRank>> NNIShape;
+		for (int i = 0; i < Shape.Num(); ++i)
+		{
+			//ORT Graph return 0 for variable dimensions, NNI use -1.
+			NNIShape.Emplace(Shape[i] == 0 ? -1 : Shape[i]);
+		}
+		
+		int Idx = AddTensor(Name, NNIShape, DataType, Data, DataSize);
+
+		return MakeTensorHandle(reinterpret_cast<void*>((int64) Idx));
+	}
+
+	/** Add model input */
+	virtual bool AddInput(HTensor Tensor) override
+	{
+		int Idx = NNETensorCast(Tensor);
+		
+		if (Idx < 0 || Idx >= Format.Tensors.Num())
+		{
+			UE_LOG(LogNNE, Warning, TEXT("Failed to add input tensor, invalid tensor index"));
+			return false;
+		}
+		
+		if (Format.Tensors[Idx].Type != EMLFormatTensorType::None)
+		{
+			UE_LOG(LogNNE, Warning, TEXT("Failed to add input tensor, tensor usage already set up"));
+			return false;
+		}
+		
+		Format.Tensors[Idx].Type = EMLFormatTensorType::Input;
+
+		return true;
+	}
+
+	/** Add model output */
+	virtual bool AddOutput(HTensor Tensor) override
+	{
+		int Idx = NNETensorCast(Tensor);
+
+		if (Idx < 0 || Idx >= Format.Tensors.Num())
+		{
+			UE_LOG(LogNNE, Warning, TEXT("Failed to add output tensor, invalid tensor index"));
+			return false;
+		}
+
+		if (Format.Tensors[Idx].Type != EMLFormatTensorType::None)
+		{
+			UE_LOG(LogNNE, Warning, TEXT("Failed to add output tensor, tensor usage already set up"));
+			return false;
+		}
+
+		Format.Tensors[Idx].Type = EMLFormatTensorType::Output;
+
+		return true;
+	}
+
+	/** Add operator */
+	virtual HOperator AddOperator(const FString& TypeName, const FString& Name = TEXT("")) override
+	{
+		int Idx = Format.Operators.Num();
+
+		FMLFormatOperatorDesc	Operator{};
+
+		Operator.TypeName = TypeName;
+		Format.Operators.Emplace(Operator);
+
+		return MakeOperatorHandle(reinterpret_cast<void*>((int64) Idx));
+	}
+
+	/** Add operator input */
+	virtual bool AddOperatorInput(HOperator Op, HTensor Tensor) override
+	{
+		int OpIdx = NNEOperatorCast(Op);
+		int TensorIdx = NNETensorCast(Tensor);
+		
+		if (TensorIdx < 0 || TensorIdx >= Format.Tensors.Num())
+		{
+			UE_LOG(LogNNE, Warning, TEXT("Failed to add operator input tensor, invalid tensor index"));
+			return false;
+		}
+
+		Format.Operators[OpIdx].InTensors.Emplace(TensorIdx);
+
+		return true;
+	}
+
+	/** Add operator output */
+	virtual bool AddOperatorOutput(HOperator Op, HTensor Tensor) override
+	{
+		int OpIdx = NNEOperatorCast(Op);
+		int TensorIdx = NNETensorCast(Tensor);
+		
+		if (TensorIdx < 0 || TensorIdx >= Format.Tensors.Num())
+		{
+			UE_LOG(LogNNE, Warning, TEXT("Failed to add operator output tensor, invalid tensor index"));
+			return false;
+		}
+
+		if (Format.Tensors[TensorIdx].Type == EMLFormatTensorType::Input)
+		{
+			UE_LOG(LogNNE, Warning, TEXT("Failed to add output tensor, tensor usage already set up to input"));
+			return false;
+		}
+
+		if (Format.Tensors[TensorIdx].Type == EMLFormatTensorType::None)
+		{
+			Format.Tensors[TensorIdx].Type = EMLFormatTensorType::Intermediate;
+		}
+		
+		Format.Operators[OpIdx].OutTensors.Emplace(TensorIdx);
+
+		return true;
+	}
+
+	/** Add operator attribute */
+	virtual bool AddOperatorAttribute(HOperator Op, const FString& Name, const FNNEAttributeValue& Value) override
+	{
+		int OpIdx = NNEOperatorCast(Op);
+
+		FMLFormatAttributeDesc& Attribute = Format.Operators[OpIdx].Attributes.Emplace_GetRef();
+		Attribute.Name = Name;
+		Attribute.Value = Value;
+		
+		return true;
+	}
+
+private:
+
+	int AddTensor(const FString& InName, TArrayView<const int32> InShape, ENNETensorDataType InDataType, const void* Data, uint64 DataSize)
+	{
+		int Idx = -1;
+
+		int* Val = TensorMap.Find(InName);
+
+		if (Val)
+		{
+			Idx = *Val;
+		}
+		else
+		{
+			FMLFormatTensorDesc	Desc{};
+
+			Desc.Name = InName;
+			Desc.Shape = InShape;
+			Desc.Type = EMLFormatTensorType::None;
+			Desc.DataType = InDataType;
+			
+			if (Data)
+			{
+				Desc.Type = EMLFormatTensorType::Initializer;
+				Desc.DataOffset = Format.TensorData.AddUninitialized(DataSize);
+				Desc.DataSize = DataSize;
+
+				FMemory::Memcpy(Format.TensorData.GetData() + Desc.DataOffset, Data, DataSize);
+			}
+			else
+			{
+				Desc.DataOffset = 0;
+				Desc.DataSize = 0;
+			}
+
+			Format.Tensors.Add(Desc);
+			Idx = Format.Tensors.Num() - 1;
+
+			TensorMap.Add(InName, Idx);
+		}
+
+		return Idx;
+	}
+
+
+	FMLRuntimeFormat		Format;
+	TMap<FString, int>		TensorMap;	
+};
+
+NNEUTILS_API IModelBuilder* CreateNNEModelBuilder()
+{
+	return new FModelBuilderNNE();
+}
+
+} // namespace UE::NNEUtils::Internal
diff --git a/Engine/Plugins/Experimental/NNE/Source/NNEUtils/Private/NNEUtilsModelBuilderONNX.cpp b/Engine/Plugins/Experimental/NNE/Source/NNEUtils/Private/NNEUtilsModelBuilderONNX.cpp
new file mode 100644
--- /dev/null
+++ b/Engine/Plugins/Experimental/NNE/Source/NNEUtils/Private/NNEUtilsModelBuilderONNX.cpp
@@ -0,0 +1,423 @@
+// Copyright Epic Games, Inc. All Rights Reserved.
+
+#include "NNXCore.h"
+#include "NNEUtilsModelBuilder.h"
+#include "NNXRuntimeFormat.h"
+#include "NNECoreAttributeMap.h"
+
+#include "NNXThirdPartyWarningDisabler.h"
+NNX_THIRD_PARTY_INCLUDES_START
+#undef check
+#undef TEXT
+
+#include "onnx/common/ir.h"
+#include "onnx/common/constants.h"
+#include "onnx/defs/operator_sets.h"
+
+#include "core/session/ort_model_optimizer_api.h"	// For model validation
+//#include "core/session/ort_env.h"
+//#include "core/session/environment.h"
+#include "core/session/onnxruntime_cxx_api.h"
+
+NNX_THIRD_PARTY_INCLUDES_END
+
+namespace UE::NNEUtils::Internal
+{
+
+inline onnx::ValueInfoProto* OnnxTensorCast(IModelBuilder::HTensor& Handle)
+{
+	if (Handle.Type == IModelBuilder::HandleType::Tensor)
+	{
+		return reinterpret_cast<onnx::ValueInfoProto*>(Handle.Ptr);
+	}
+
+	return nullptr;
+}
+
+inline onnx::NodeProto* OnnxOperatorCast(IModelBuilder::HOperator& Handle)
+{
+	if (Handle.Type == IModelBuilder::HandleType::Operator)
+	{
+		return reinterpret_cast<onnx::NodeProto*>(Handle.Ptr);
+	}
+
+	return nullptr;
+}
+
+/**
+ * Builds an ONNX model in memory.
+ * NOTE:
+ * - We plan to use this only for generating simple networks for testing operators and simple models
+ */
+class FMLModelBuilderONNX : public IModelBuilder
+{
+	static constexpr const char* kOnnxDomain = onnx::ONNX_DOMAIN;
+
+	onnx::ModelProto	Model;
+	onnx::GraphProto*	Graph{ nullptr };
+	int64				IrVersion { OnnxIrVersion };
+	int64				OpsetVersion{ OnnxOpsetVersion };
+	
+public:
+
+	FMLModelBuilderONNX(int64 InIrVersion = OnnxIrVersion, int64 InOpsetVersion = OnnxOpsetVersion)
+		: IrVersion(InIrVersion)
+		, OpsetVersion(InOpsetVersion)
+	{
+	}
+	
+	/** Initialize the model builder */
+	virtual bool Begin(const FString& Name) override
+	{
+		// Setup model
+		Model.set_ir_version(IrVersion);
+		Model.set_domain(kOnnxDomain);
+
+		onnx::OperatorSetIdProto* OpsetProto = Model.add_opset_import();
+
+		OpsetProto->set_domain(kOnnxDomain);
+		OpsetProto->set_version(OpsetVersion);
+
+		Graph = Model.mutable_graph();
+		Graph->set_name(TCHAR_TO_ANSI(*Name));
+
+		return true;
+	}
+
+	/** Serialize the model to given array */
+	virtual bool End(TArray<uint8>& Data) override
+	{
+		const int Size = (int)Model.ByteSizeLong();
+
+		//Data.r
+		Data.SetNum(Size);
+
+		bool res = Model.SerializeToArray(Data.GetData(), Data.Num());
+
+		// Validate model
+		if (res)
+		{
+			UE_LOG(LogNNX, Display, TEXT("OrtValidateModelFromMemory... (note: might abort if i.e. shapes are not correct!)"));
+
+			OrtStatusPtr status = OrtValidateModelFromMemory(Data.GetData(), Data.Num());
+
+			if (status)
+			{
+				UE_LOG(LogNNX, Warning, TEXT("ModelBuilder error:%s"), ANSI_TO_TCHAR(Ort::GetApi().GetErrorMessage(status)));
+				return false;
+			}
+		}
+
+		return res;
+	}
+
+	/** Add tensor */
+	virtual HTensor AddTensor(const FString& Name, ENNETensorDataType DataType, TArrayView<const int32> Shape, const void* Data, uint64 DataSize) override
+	{
+		auto Value = onnx::ValueInfoProto().New(Graph->GetArena());
+		
+		if (!SetValue(Value, Name, DataType, Shape))
+		{
+			return HTensor();
+		}
+
+		if (Data)
+		{
+			onnx::TensorProto* Tensor = Graph->mutable_initializer()->Add();
+
+			Tensor->set_name(TCHAR_TO_ANSI(*Name));
+			Tensor->set_data_type(ToTensorProtoDataType(DataType));
+			for (int32 Idx = 0; Idx < Shape.Num(); ++Idx)
+			{
+				Tensor->add_dims(Shape[Idx]);
+			}
+
+			checkf(DataType != ENNETensorDataType::Char, TEXT("Char tensors not supported yet!"));
+			// raw_data is not supported for strings, when implementing those one will need to use string_data from TensorProto
+			std::string* raw_data = Tensor->mutable_raw_data();
+			raw_data->assign((const char*)Data, DataSize);
+		}
+
+		return MakeTensorHandle(Value);
+	}
+
+	/** Add model input */
+	virtual bool AddInput(HTensor Handle) override
+	{
+		auto Value = OnnxTensorCast(Handle);
+
+		if (!Value)
+		{
+			return false;
+		}
+
+		*Graph->mutable_input()->Add() = *Value;
+
+		return true;
+	}
+
+	/** Add model output */
+	virtual bool AddOutput(HTensor Handle) override
+	{
+		auto Value = OnnxTensorCast(Handle);
+
+		if (!Value)
+		{
+			return false;
+		}
+
+		*Graph->mutable_output()->Add() = *Value;
+
+		return true;
+	}
+
+	/** Add operator */
+	virtual HOperator AddOperator(const FString& TypeName, const FString& Name) override
+	{
+		FTCHARToUTF8 Convert(*TypeName);
+
+		std::string NodeType = onnx::Symbol(Convert.Get()).toString();
+
+		onnx::NodeProto* Node = Graph->add_node();
+
+		Node->set_op_type(NodeType);
+
+		if (!Name.IsEmpty())
+		{
+			Node->set_name(TCHAR_TO_ANSI(*Name));
+		}
+		else
+		{
+			Node->set_name(TCHAR_TO_ANSI(*TypeName));
+		}
+
+		Node->set_domain(Model.domain());
+
+		return MakeOperatorHandle(Node);
+	}
+
+	/** Add operator input */
+	virtual bool AddOperatorInput(HOperator Op, HTensor Tensor) override
+	{
+		auto Value = OnnxTensorCast(Tensor);
+		auto NodeOp = OnnxOperatorCast(Op);
+		auto InName = NodeOp->mutable_input()->Add();
+		
+		*InName = Value->name();
+		
+		return true;
+	}
+
+	/** Add operator output */
+	virtual bool AddOperatorOutput(HOperator Op, HTensor Tensor) override
+	{
+		auto Value = OnnxTensorCast(Tensor);
+		auto NodeOp = OnnxOperatorCast(Op);
+		auto OutName = NodeOp->mutable_output()->Add();
+
+		*OutName = Value->name();
+
+		return true;
+	}
+
+	/** Add operator attribute */
+	virtual bool AddOperatorAttribute(HOperator Op, const FString& Name, const FNNEAttributeValue& Value) override
+	{
+		auto NodeOp = OnnxOperatorCast(Op);
+		onnx::AttributeProto* Attribute = NodeOp->mutable_attribute()->Add();
+
+		Attribute->set_name(TCHAR_TO_ANSI(*Name));
+		if (Value.GetType() == ENNEAttributeDataType::Float)
+		{
+			Attribute->set_type(onnx::AttributeProto::FLOAT);
+			Attribute->set_f(Value.GetValue<float>());
+		}
+		else if (Value.GetType() == ENNEAttributeDataType::Int32)
+		{
+			Attribute->set_type(onnx::AttributeProto::INT);
+			Attribute->set_i(Value.GetValue<int32>());
+		}
+		else if (Value.GetType() == ENNEAttributeDataType::Int32Array)
+		{
+			Attribute->set_type(onnx::AttributeProto::INTS);
+			for (int32 Val : Value.GetValue<TArray<int32>>())
+			{
+				Attribute->add_ints(Val);
+			}
+		}
+		else if (Value.GetType() == ENNEAttributeDataType::String)
+		{
+			Attribute->set_type(onnx::AttributeProto::STRING);
+			Attribute->set_s(TCHAR_TO_ANSI(*Value.GetValue<FString>()));
+		}
+		else
+		{
+			checkf(false, TEXT("not implemented"))
+		}
+		
+		return true;
+	}
+
+private:
+
+	bool SetValue(onnx::ValueInfoProto* Value, const FString& Name, ENNETensorDataType DataType, const TArrayView<const int32>& InShape)
+	{
+		onnx::TypeProto*		Type = Value->mutable_type();
+		onnx::TypeProto_Tensor* TensorType = Type->mutable_tensor_type();
+		onnx::TensorShapeProto* Shape = TensorType->mutable_shape();
+		
+		Value->set_name(TCHAR_TO_ANSI(*Name));
+		TensorType->set_elem_type(ToTensorProtoDataType(DataType));
+
+		for (int32 Idx = 0; Idx < InShape.Num(); ++Idx)
+		{
+			auto Dim = Shape->add_dim();
+
+			Dim->set_dim_value(InShape[Idx]);
+		}
+
+		return true;
+	}
+
+	onnx::TensorProto_DataType ToTensorProtoDataType(ENNETensorDataType DataType)
+	{
+		switch (DataType)
+		{
+			case ENNETensorDataType::None: return onnx::TensorProto_DataType_UNDEFINED;
+			case ENNETensorDataType::Float: return onnx::TensorProto_DataType_FLOAT;
+			case ENNETensorDataType::UInt8: return onnx::TensorProto_DataType_UINT8;
+			case ENNETensorDataType::Int8: return onnx::TensorProto_DataType_INT8;
+			case ENNETensorDataType::UInt16: return onnx::TensorProto_DataType_UINT16;
+			case ENNETensorDataType::Int16: return onnx::TensorProto_DataType_INT16;
+			case ENNETensorDataType::Int32: return onnx::TensorProto_DataType_INT32;
+			case ENNETensorDataType::Int64: return onnx::TensorProto_DataType_INT64;
+			//case ENNETensorDataType::String: return onnx::TensorProto_DataType_STRING;
+			case ENNETensorDataType::Boolean: return onnx::TensorProto_DataType_BOOL;
+			case ENNETensorDataType::Half: return onnx::TensorProto_DataType_FLOAT16;
+			case ENNETensorDataType::Double: return onnx::TensorProto_DataType_DOUBLE;
+			case ENNETensorDataType::UInt32: return onnx::TensorProto_DataType_UINT32;
+			case ENNETensorDataType::UInt64: return onnx::TensorProto_DataType_UINT64;
+			case ENNETensorDataType::Complex64: return onnx::TensorProto_DataType_COMPLEX64;
+			case ENNETensorDataType::Complex128: return onnx::TensorProto_DataType_COMPLEX128;
+			case ENNETensorDataType::BFloat16: return onnx::TensorProto_DataType_BFLOAT16;
+			default: return onnx::TensorProto_DataType_UNDEFINED;
+		}
+	}
+};
+
+//
+//
+//
+void BuildShapeForModel(bool ConvertToVariadicShape, const NNECore::FTensorShape& InShape, TArray<int32>& OutShape)
+{
+	OutShape.Empty();
+	for (int32 Idx = 0; Idx < InShape.Rank(); ++Idx)
+	{
+		int32 Dim = static_cast<int32>(InShape.GetData()[Idx]);
+		if (ConvertToVariadicShape)
+		{
+			Dim = -1;
+		}
+		OutShape.Emplace(Dim);
+	}
+}
+
+//
+//
+//
+NNEUTILS_API bool CreateONNXModelForOperator(bool UseVariadicShapeForModel, const FString& OperatorName,
+	TConstArrayView<FTensor> InInputTensors, TConstArrayView<FTensor> InOutputTensors,
+	TConstArrayView<FTensor> InWeightTensors, TConstArrayView<TArray<char>> InWeightTensorsData,
+	const UE::NNECore::FAttributeMap& Attributes, FNNIModelRaw& Model)
+{
+	Model = FNNIModelRaw{};
+	
+	int64 IrVersion = OnnxIrVersion;
+	int64 OpsetVersion = OnnxOpsetVersion;
+	if (OperatorName == TEXT("Upsample") || OperatorName == TEXT("Pad"))
+	{
+		// Upsample is deprecated starting opset 10
+		// Pad current implementation is opset 2 (next version is 11)
+		OpsetVersion = 9;
+	}
+	TUniquePtr<IModelBuilder> Builder(CreateONNXModelBuilder(IrVersion, OpsetVersion));
+
+	Builder->Begin();
+
+	TArray<int32> ShapeForModel;
+	TArray<IModelBuilder::HTensor> InputTensors;
+	
+	for (int32 Idx = 0; Idx < InInputTensors.Num(); ++Idx)
+	{
+		const FTensor& Desc = InInputTensors[Idx];
+		BuildShapeForModel(UseVariadicShapeForModel, Desc.GetShape(), ShapeForModel);
+		IModelBuilder::HTensor Tensor = Builder->AddTensor(Desc.GetName(), Desc.GetDataType(), ShapeForModel);
+
+		InputTensors.Emplace(Tensor);
+		Builder->AddInput(Tensor);
+	}
+
+	TArray<IModelBuilder::HTensor> OutputTensors;
+
+	for (int32 Idx = 0; Idx < InOutputTensors.Num(); ++Idx)
+	{
+		const FTensor& Desc = InOutputTensors[Idx];
+		BuildShapeForModel(UseVariadicShapeForModel, Desc.GetShape(), ShapeForModel);
+		IModelBuilder::HTensor Tensor = Builder->AddTensor(Desc.GetName(), Desc.GetDataType(), ShapeForModel);
+
+		OutputTensors.Emplace(Tensor);
+		Builder->AddOutput(Tensor);
+	}
+
+	checkf(InWeightTensors.Num() == InWeightTensorsData.Num(), TEXT("Invalid weight tensors data"));
+	TArray<IModelBuilder::HTensor> WeightTensors;
+
+	for (int32 Idx = 0; Idx < InWeightTensors.Num(); ++Idx)
+	{
+		const FTensor& Desc = InWeightTensors[Idx];
+		const TArray<char>& Data = InWeightTensorsData[Idx];
+		check(Data.Num() == Desc.GetDataSize());
+		BuildShapeForModel(false, Desc.GetShape(), ShapeForModel);
+		IModelBuilder::HTensor Tensor = Builder->AddTensor(Desc.GetName(), Desc.GetDataType(), ShapeForModel, Data.GetData(), Data.Num());
+
+		WeightTensors.Emplace(Tensor);
+	}
+
+	auto Op = Builder->AddOperator(OperatorName);
+
+	for (int32 Idx = 0; Idx < InputTensors.Num(); ++Idx)
+	{
+		Builder->AddOperatorInput(Op, InputTensors[Idx]);
+	}
+
+	for (int32 Idx = 0; Idx < WeightTensors.Num(); ++Idx)
+	{
+		// For now weights are added after model inputs in the list of
+		// operator inputs. This should be made more flexible if needed by future tests cases.
+		Builder->AddOperatorInput(Op, WeightTensors[Idx]);
+	}
+
+	for (int32 Idx = 0; Idx < OutputTensors.Num(); ++Idx)
+	{
+		Builder->AddOperatorOutput(Op, OutputTensors[Idx]);
+	}
+
+	for (int32 Idx = 0; Idx < Attributes.Num(); ++Idx)
+	{
+		Builder->AddOperatorAttribute(Op, Attributes.GetName(Idx), Attributes.GetAttributeValue(Idx));
+	}
+
+	
+	Builder->End(Model.Data);
+	Model.Format = ENNXInferenceFormat::ONNX;
+	
+	return true;
+}
+
+/** Return instance of ONNX model builder */
+NNEUTILS_API IModelBuilder* CreateONNXModelBuilder(int64 IrVersion, int64 OpsetVersion)
+{
+	return new FMLModelBuilderONNX(IrVersion, OpsetVersion);
+}
+
+} // namespace UE::NNEUtils::Internal
diff --git a/Engine/Plugins/Experimental/NNE/Source/NNEUtils/Private/NNEUtilsModelOptimizer.cpp b/Engine/Plugins/Experimental/NNE/Source/NNEUtils/Private/NNEUtilsModelOptimizer.cpp
new file mode 100644
--- /dev/null
+++ b/Engine/Plugins/Experimental/NNE/Source/NNEUtils/Private/NNEUtilsModelOptimizer.cpp
@@ -0,0 +1,665 @@
+// Copyright Epic Games, Inc. All Rights Reserved.
+
+#include "NNEUtilsModelOptimizer.h"
+#include "NNECore.h"
+#include "NNEUtilsModelBuilder.h"
+#include "NNXModelOptimizerInterface.h"
+#include "NNXRuntimeFormat.h"
+
+
+#ifdef PLATFORM_NNX_MICROSOFT
+// Prevent some of false-positives when doing static analysis
+__pragma(warning(disable: 6385))
+#endif
+
+#if PLATFORM_WINDOWS
+#define ORT_STRING_CAST TCHAR_TO_WCHAR
+#else
+#define ORT_STRING_CAST TCHAR_TO_ANSI
+#endif
+
+#include "NNXThirdPartyWarningDisabler.h"
+NNX_THIRD_PARTY_INCLUDES_START
+#undef check
+#undef TEXT
+#include "core/session/ort_model_optimizer_api.h"
+#include "core/session/onnxruntime_cxx_api.h"
+NNX_THIRD_PARTY_INCLUDES_END
+
+
+#define Print(Format, ...) UE_LOG(LogNNE, Display, Format, __VA_ARGS__)
+
+class FModelGraphPrinter
+{
+public:
+
+	FModelGraphPrinter(Ort::IModelGraph* InGraph)
+		: Graph(InGraph)
+	{
+		UE_LOG(LogNNE, Display, TEXT("Visiting model:%s"), ANSI_TO_TCHAR(Graph->GetGraphInfo().name));
+
+		Storage.SetNum(2048);
+	}
+
+	void Run()
+	{
+		auto GraphInfo = Graph->GetGraphInfo();	
+
+		Print(TEXT("Graph:%s"), ANSI_TO_TCHAR(GraphInfo.name));
+        Print(TEXT("- Inputs:%d"), GraphInfo.inputCount);
+		for (int c = 0; c < GraphInfo.inputCount; ++c) 
+		{
+			VisitTensor(Graph->GetGraphInput(c));
+		}
+
+		Print(TEXT("- Outputs:%d"), GraphInfo.outputCount);
+		for (int c = 0; c < GraphInfo.outputCount; ++c) 
+		{
+			VisitTensor(Graph->GetGraphOutput(c));
+		}
+
+		Print(TEXT("- Nodes:%d"), GraphInfo.nodeCount);
+		for (int NodeIdx = 0; NodeIdx < GraphInfo.nodeCount; ++NodeIdx) 
+		{
+			VisitNode(Graph->GetNode(NodeIdx));
+		}
+        
+		Print(TEXT("- Tensor initializers:%d"), GraphInfo.tensorInitializerCount);
+	}
+
+	void VisitNode(const Ort::GraphNode& Node)
+	{
+		Ort::GraphNodeInfo	NodeInfo = Graph->GetNodeInfo(Node);
+
+		Print(TEXT("Node op:%s name:%s"), ANSI_TO_TCHAR(NodeInfo.opName), ANSI_TO_TCHAR(NodeInfo.name));
+
+		Print(TEXT("- Attribs:%d"), NodeInfo.attributeCount);
+		for (int AttrIdx = 0; AttrIdx < NodeInfo.attributeCount; ++AttrIdx)
+		{
+			VisitAttrib(Node, AttrIdx);
+			//Visit(Graph->GetNodeAttributeValue(Node, AttrIdx));
+		}
+
+		Print(TEXT("- Inputs:%d"), NodeInfo.inputCount);
+		for (int InIdx = 0; InIdx < NodeInfo.inputCount; ++InIdx)
+		{
+			VisitTensor(Graph->GetNodeInput(Node, InIdx));
+		}
+
+		Print(TEXT("- Outputs:%d"), NodeInfo.outputCount);
+		for (int OutIdx = 0; OutIdx < NodeInfo.outputCount; ++OutIdx)
+		{
+			VisitTensor(Graph->GetNodeOutput(Node, OutIdx));
+		}
+	}
+
+	void VisitAttrib(const Ort::GraphNode& Node, int AttrIdx)
+	{
+		Ort::GraphAttributeInfo Attrib = Graph->GetNodeAttribute(Node, AttrIdx);
+		Ort::GraphAttributeValue Value = Graph->GetNodeAttributeValue(Node, AttrIdx);
+
+		Print(TEXT("   %s %d"), ANSI_TO_TCHAR(Attrib.name), Attrib.type);
+		
+		if (Value.type == Ort::GraphAttributeValue::kFloat)
+		{
+            Print(TEXT("      %f"), Value.f);
+		}
+        else if (Value.type == Ort::GraphAttributeValue::kInt)
+		{
+			Print(TEXT("      %lld"), Value.i);
+		}
+		else if (Value.type == Ort::GraphAttributeValue::kString)
+		{
+			Print(TEXT("      %s"), Value.s);
+		}
+        else if (Value.type == Ort::GraphAttributeValue::kFloats)
+		{
+            for (int c = 0; c < Value.count; ++c)
+			{
+				Print(TEXT("      %f"), Value.floats[c]);
+			}
+		}
+		else if (Value.type == Ort::GraphAttributeValue::kInts) 
+		{
+			for (int c = 0; c < Value.count; ++c)
+			{
+				Print(TEXT("      %lld"), Value.ints[c]);
+			}
+		}
+		else
+		{
+			UE_LOG(LogNNE, Warning, TEXT("Unsupported attribute value type"));
+		}
+	}
+
+	void VisitTensor(const Ort::GraphTensorInfo& Tensor) 
+	{
+		FAnsiStringBuilderBase	Str;
+
+		Str.Appendf("   %-50s ", Tensor.name);
+		Str.Append(" [ ");
+
+		for (int c = 0; c < Tensor.shapeLen; ++c) 
+		{
+			if (Tensor.shape[c] == 0)
+				Str.Append("N");
+			else
+				Str.Appendf("%d", Tensor.shape[c]);
+
+			Str.Appendf("%c", c < Tensor.shapeLen - 1 ? ',' : ' ');
+		}
+		
+		Str.Append("]");
+		Str.Appendf(" type:%s", Ort::GraphTensorDataTypeToString(Tensor.dataType));
+		
+		Ort::GraphTensorInitializer	TensorInit = Graph->GetTensorInitializer(Tensor.name);
+		size_t						DataSize = 0;
+
+		if (TensorInit) 
+		{
+			DataSize = Graph->GetTensorDataSize(TensorInit);
+			if (DataSize > Storage.Num())
+			{
+				Storage.SetNum(DataSize);
+			}
+
+			Str.Appendf(" size:%d", (int) DataSize);
+		}
+
+		Str.AppendChar('\0');
+		UE_LOG(LogNNE, Display, TEXT("%s"), ANSI_TO_TCHAR(Str.GetData()));
+		Str.Reset();
+		
+		if (DataSize)
+		{
+			if (Graph->GetTensorData(TensorInit, Storage.GetData(), DataSize, 0) == 0) 
+			{
+				int NElems = 10;
+
+				if (Tensor.shapeLen == 1) 
+				{
+					if (Tensor.shape[0] < NElems)
+					{
+                        NElems = Tensor.shape[0];
+					}
+				} 
+				else 
+				{
+                    const int dim = Tensor.shape[Tensor.shapeLen - 1];
+
+					if (dim < NElems)
+					{
+						NElems = dim;
+					}
+				}
+
+				for (int ec = 0; ec < NElems; ++ec) 
+				{
+					if (Tensor.dataType == Ort::GraphTensorDataType::kFloat)
+					{
+						UE_LOG(LogNNE, Display, TEXT("      %f"), ((float*) Storage.GetData())[ec]);
+					}
+                    else if (Tensor.dataType == Ort::GraphTensorDataType::kInt32)
+					{
+						UE_LOG(LogNNE, Display, TEXT("      %d"), ((int32*) Storage.GetData())[ec]);
+					}
+					else if (Tensor.dataType == Ort::GraphTensorDataType::kUInt32)
+					{
+						UE_LOG(LogNNE, Display, TEXT("      %u"), ((uint32*) Storage.GetData())[ec]);
+					}
+                    else if (Tensor.dataType == Ort::GraphTensorDataType::kInt64)
+					{
+						UE_LOG(LogNNE, Display, TEXT("      %lld"), ((int64*)Storage.GetData())[ec]);
+					}
+					else if (Tensor.dataType == Ort::GraphTensorDataType::kUInt64)
+					{
+						UE_LOG(LogNNE, Display, TEXT("      %llu"), ((uint64*) Storage.GetData())[ec]);
+					}
+				}
+			}
+			else 
+			{
+                UE_LOG(LogNNE, Warning, TEXT("Failed to read tensor data :%s"), Tensor.name);
+			}
+		}
+	}
+
+private:
+	TUniquePtr<Ort::IModelGraph>	Graph;
+	TArray<uint8>					Storage;
+};
+
+#undef Print
+
+
+#define CASE(GType, Type) case Ort::GraphTensorDataType::GType: return ENNETensorDataType::Type
+
+ENNETensorDataType GetDataTypeFromGraphTensor(Ort::GraphTensorDataType TensorDataType)
+{
+	switch (TensorDataType)
+	{
+		CASE(kFloat, Float);
+		CASE(kUInt8, UInt8);
+		CASE(kInt8, Int8);
+		CASE(kUInt16, UInt16);
+		CASE(kInt16, Int16);
+		CASE(kInt32, Int32);
+		CASE(kInt64, Int64);
+		//CASE(kString, String);
+		CASE(kBool, Boolean);
+		CASE(kFloat16, Half);
+		CASE(kDouble, Double);
+		CASE(kUInt32, UInt32);
+		CASE(kUInt64, UInt64);
+		CASE(kComplex64, Complex64);
+		CASE(kComplex128, Complex128);
+		CASE(kBFloat16, BFloat16);
+
+		default:
+			return ENNETensorDataType::None;
+	}
+}
+
+#undef CASE
+
+namespace UE::NNEUtils::Internal
+{
+using namespace NNX;
+
+class FModelOptimizerBase : public IModelOptimizer
+{
+protected:
+	TArray<TSharedPtr<IModelOptimizerPass>> OptimizationPasses;
+	TArray<TSharedPtr<IModelValidator>> Validators;
+
+	virtual bool ValidateInputModel(const FNNIModelRaw& InputModel)
+	{
+		if (InputModel.Format != ENNXInferenceFormat::ONNX)
+		{
+			UE_LOG(LogNNE, Warning, TEXT("Optimizer %s is expecting ONNX input format."), *GetName());
+			return false;
+		}
+
+		OrtStatusPtr Status = OrtValidateModelFromMemory(InputModel.Data.GetData(), InputModel.Data.Num());
+		if (Status)
+		{
+			UE_LOG(LogNNE, Warning, TEXT("Input ONNX model is invalid: %s, Model won't be optimized"), ANSI_TO_TCHAR(Ort::GetApi().GetErrorMessage(Status)));
+			return false;
+		}
+		
+		return true;
+	}
+	
+public:
+	virtual void AddOptimizationPass(TSharedPtr<IModelOptimizerPass> ModelOptimizerPass) override
+	{
+		if (ModelOptimizerPass.IsValid())
+		{
+			OptimizationPasses.Add(ModelOptimizerPass);
+		}
+	}
+
+	virtual void AddValidator(TSharedPtr<IModelValidator> ModelValidator) override
+	{
+		if (ModelValidator.IsValid())
+		{
+			Validators.Add(ModelValidator);
+		}
+	}
+
+	bool IsModelValid(const FNNIModelRaw& ModelToValidate, const FOptimizerOptionsMap& Options)
+	{
+		bool bIsModelValid = true;
+
+		for (TSharedPtr<IModelValidator>& Validator : Validators)
+		{
+			check(Validator.IsValid());
+			if (!Validator->ValidateModel(ModelToValidate, Options))
+			{
+				UE_LOG(LogNNE, Warning, TEXT("Model validator %s detected an error."), *(Validator->GetName()));
+				bIsModelValid = false;
+			}
+		}
+		return bIsModelValid;
+	}
+
+	bool ApplyAllPassesAndValidations(FNNIModelRaw& OptimizedModel, const FOptimizerOptionsMap& Options)
+	{
+		if (!IsModelValid(OptimizedModel, Options))
+		{
+			UE_LOG(LogNNE, Warning, TEXT("Model is not valid, skipping optimization passes."));
+			return false;
+		}
+		
+		for (TSharedPtr<IModelOptimizerPass>& Pass : OptimizationPasses)
+		{
+			check(Pass.IsValid());
+			
+			if (!Pass->ApplyPass(OptimizedModel, Options))
+			{
+				UE_LOG(LogNNE, Warning, TEXT("Error while executing model optimisation pass %s."), *(Pass->GetName()));
+				return false;
+			}
+			if (!IsModelValid(OptimizedModel, Options))
+			{
+				UE_LOG(LogNNE, Warning, TEXT("Model validation failed after optimisation pass %s."), *(Pass->GetName()));
+				return false;
+			}
+		}
+
+		return true;
+	}
+
+	bool Optimize(const FNNIModelRaw& InputModel, FNNIModelRaw& OptimizedModel, const FOptimizerOptionsMap& Options) override
+	{
+		OptimizedModel = FNNIModelRaw{};
+		
+		if (!ValidateInputModel(InputModel))
+		{
+			return false;
+		}
+
+		OptimizedModel = InputModel;
+		return ApplyAllPassesAndValidations(OptimizedModel, Options);
+	}
+};
+
+class FONNXModelValidator : public IModelValidator
+{
+public:
+	virtual FString GetName() const
+	{
+		return TEXT("ONNX Model validator");
+	}
+
+	virtual bool ValidateModel(const FNNIModelRaw& InputModel, const FOptimizerOptionsMap& Options) const override
+	{
+		FMLRuntimeFormat	Format;
+
+		ENNXInferenceFormat FormatType = InputModel.Format;
+		if (FormatType != ENNXInferenceFormat::ONNX)
+		{
+			UE_LOG(LogNNE, Warning, TEXT("Unsupported format type for validator %s"), *GetName());
+			return false;
+		}
+
+		OrtStatusPtr Status = OrtValidateModelFromMemory(InputModel.Data.GetData(), InputModel.Data.Num());
+		if (Status)
+		{
+			UE_LOG(LogNNE, Warning, TEXT("Failed to validate ONNX model: %s"), ANSI_TO_TCHAR(Ort::GetApi().GetErrorMessage(Status)));
+			return false;
+		}
+
+		return true;
+	}
+};
+
+class FModelOptimizerONNXToONNX : public FModelOptimizerBase
+{
+public:
+	virtual FString GetName() const override
+	{
+		return TEXT("NNEModelOptimizerFromONNXToONNX");
+	}
+
+	//TODO jira 167591: investigate if it is possible to optimize ONNX model and serialize it back as ONNX.
+	//Among other benefits applying L1 could cut down on the model size. Atm this optimizer is a pass trough.
+	//Bonus: Can we get a validator for onnx format?
+};
+
+static void OnOrtLog(const char* LogMsg)
+{
+	UE_LOG(LogNNE, Warning, TEXT("%s"), ANSI_TO_TCHAR(LogMsg));
+}
+
+class FModelOptimizerONNXToORT : public FModelOptimizerBase
+{
+public:
+	virtual FString GetName() const override
+	{
+		return TEXT("NNEModelOptimizerONNXToORT");
+	}
+
+	virtual bool Optimize(const FNNIModelRaw& InputModel, FNNIModelRaw& OptimizedModel, const FOptimizerOptionsMap& Options) override
+	{
+		OptimizedModel = FNNIModelRaw{};
+		
+		if (!ValidateInputModel(InputModel))
+		{
+			return false;
+		}
+
+		TUniquePtr<Ort::IModelGraph> Graph = ConvertONNXToORTModelGraph(InputModel.Data);
+
+		//TODO jira 167588: Serialize the Graph to a buffer and store it in OutModel.Data
+		//allowing to initiate a session from ORT format for NN engine supporting this.
+		//Allowing model optimization to be applyed at cooking time.
+		//OptimizedModel = Graph.Serialize;
+		//return ApplyAllPasses(OptimizedModel, Options);
+		return false;
+	}
+
+protected:
+	TUniquePtr<Ort::IModelGraph> ConvertONNXToORTModelGraph(TConstArrayView<uint8> ONNXData)
+	{
+		Ort::ModelOptimizeOptions	Opts{};
+
+		Opts.logCallback = OnOrtLog;
+
+		TUniquePtr<Ort::IModelGraph> Graph(OrtOptimizeModelFromMemory(ONNXData.GetData(), ONNXData.Num(), Opts));
+		if (!Graph)
+		{
+			UE_LOG(LogNNE, Warning, TEXT("Failed to load ONNX model from memory"));
+		}
+
+		//auto Printer = MakeUnique<FModelGraphPrinter>(Graph.Get());
+		//Printer->Run();
+
+		return Graph;
+	}
+};
+
+class FModelOptimizerONNXToNNERT : public FModelOptimizerONNXToORT
+{
+public:
+	virtual FString GetName() const override
+	{
+		return TEXT("FNNEModelOptimizerONNXToNNE");
+	}
+	
+	virtual bool Optimize(const FNNIModelRaw& InputModel, FNNIModelRaw& OptimizedModel, const FOptimizerOptionsMap& Options) override
+	{
+		OptimizedModel = FNNIModelRaw{};
+		
+		if (!ValidateInputModel(InputModel))
+		{
+			return false;
+		}
+
+		TUniquePtr<Ort::IModelGraph> Graph = ConvertONNXToORTModelGraph(InputModel.Data);
+		
+		if (!BuildNNEFormat(Graph.Get(), OptimizedModel.Data))
+		{
+			UE_LOG(LogNNE, Warning, TEXT("Error while building ORT ModelGraph."));
+			return false;
+		}
+
+		OptimizedModel.Format = ENNXInferenceFormat::NNXRT;
+
+		return ApplyAllPassesAndValidations(OptimizedModel, Options);
+	}
+
+private:
+	bool BuildNNEFormat(Ort::IModelGraph* Graph, TArray<uint8>& NNEData)
+	{
+		TUniquePtr<IModelBuilder>	Builder(CreateNNEModelBuilder());
+
+		Ort::GraphInfo GraphInfo = Graph->GetGraphInfo();
+
+		Builder->Begin(GraphInfo.name);
+
+		// Add tensors for graph inputs
+		for (int Idx = 0; Idx < GraphInfo.inputCount; ++Idx)
+		{
+			Ort::GraphTensorInfo TensorInfo = Graph->GetGraphInput(Idx);
+
+			ENNETensorDataType DataType = GetDataTypeFromGraphTensor(TensorInfo.dataType);
+
+			auto Tensor = 
+				Builder->AddTensor(
+					ANSI_TO_TCHAR(TensorInfo.name),
+					DataType,
+					MakeArrayView(TensorInfo.shape, TensorInfo.shapeLen)
+				);
+
+			Builder->AddInput(Tensor);
+		}
+
+		// Add tensors for graph outputs
+		for (int Idx = 0; Idx < GraphInfo.outputCount; ++Idx)
+		{
+			Ort::GraphTensorInfo TensorInfo = Graph->GetGraphOutput(Idx);
+
+			ENNETensorDataType DataType = GetDataTypeFromGraphTensor(TensorInfo.dataType);
+
+			auto Tensor =
+				Builder->AddTensor(
+					ANSI_TO_TCHAR(TensorInfo.name),
+					DataType,
+					MakeArrayView(TensorInfo.shape, TensorInfo.shapeLen)
+				);
+
+			Builder->AddOutput(Tensor);
+		}
+
+		// Traverse all the nodes get their inputs, outputs and tensor data
+		TArray<uint8> TensorDataBlob;
+
+		for (int Idx = 0; Idx < GraphInfo.nodeCount; ++Idx)
+		{
+			Ort::GraphNode		Node = Graph->GetNode(Idx);
+			Ort::GraphNodeInfo	NodeInfo = Graph->GetNodeInfo(Node);
+
+			auto Op = Builder->AddOperator(NodeInfo.opName);
+
+			for (int InIdx = 0; InIdx < NodeInfo.attributeCount; ++InIdx)
+			{
+				Ort::GraphAttributeInfo AttrInfo = Graph->GetNodeAttribute(Node, InIdx);
+				Ort::GraphAttributeValue AttrValue = Graph->GetNodeAttributeValue(Node, InIdx);
+
+				if (AttrInfo.type == Ort::GraphAttributeType::kFloat)
+				{
+					Builder->AddOperatorAttribute(Op, AttrInfo.name, FNNEAttributeValue(AttrValue.f));
+				}
+				else if (AttrInfo.type == Ort::GraphAttributeType::kInt)
+				{
+					Builder->AddOperatorAttribute(Op, AttrInfo.name, FNNEAttributeValue((int32)AttrValue.i));
+				}
+				else if (AttrInfo.type == Ort::GraphAttributeType::kInts)
+				{
+					TArray<int32> Values;
+					Values.SetNumUninitialized(AttrValue.count);
+
+					for (int i = 0; i < AttrValue.count; i++)
+					{
+						Values[i] = (int32)AttrValue.ints[i];
+					}
+
+					Builder->AddOperatorAttribute(Op, AttrInfo.name, FNNEAttributeValue(Values));
+				}
+				else if (AttrInfo.type == Ort::GraphAttributeType::kString)
+				{
+					Builder->AddOperatorAttribute(Op, AttrInfo.name, FNNEAttributeValue(FString(AttrValue.s)));
+				}
+				else
+				{
+					//TODO: better error reporting add type (example: sparse tensor) and name of the actual node if any (not the op name but the node one)
+					UE_LOG(LogNNE, Warning, TEXT("Unsupported attribute type for attribute '%s' in node '%s'"), ANSI_TO_TCHAR(AttrInfo.name), ANSI_TO_TCHAR(NodeInfo.opName));
+				}
+			}
+
+			for (int InIdx = 0; InIdx < NodeInfo.inputCount; ++InIdx)
+			{
+				Ort::GraphTensorInfo		TensorInfo = Graph->GetNodeInput(Node, InIdx);
+				Ort::GraphTensorInitializer	TensorInit = Graph->GetTensorInitializer(TensorInfo.name);
+				ENNETensorDataType			DataType = GetDataTypeFromGraphTensor(TensorInfo.dataType);
+				const void*					Data = nullptr;
+				uint64						DataSize = 0;
+
+				if (TensorInit)
+				{
+					DataSize = Graph->GetTensorDataSize(TensorInit);
+
+					if (DataSize)
+					{
+						if (TensorDataBlob.Num() < DataSize)
+						{
+							TensorDataBlob.SetNumUninitialized(DataSize);
+						}
+
+						Graph->GetTensorData(TensorInit, TensorDataBlob.GetData(), DataSize, 0);
+						Data = TensorDataBlob.GetData();
+					}
+				}
+
+				auto Tensor =
+					Builder->AddTensor(
+						ANSI_TO_TCHAR(TensorInfo.name),
+						DataType,
+						MakeArrayView(TensorInfo.shape, TensorInfo.shapeLen),
+						Data, DataSize
+					);
+
+				Builder->AddOperatorInput(Op, Tensor);
+			}
+
+			for (int OutIdx = 0; OutIdx < NodeInfo.outputCount; ++OutIdx)
+			{
+				Ort::GraphTensorInfo		TensorInfo = Graph->GetNodeOutput(Node, OutIdx);
+				ENNETensorDataType			DataType = GetDataTypeFromGraphTensor(TensorInfo.dataType);
+				
+				auto Tensor =
+					Builder->AddTensor(
+						ANSI_TO_TCHAR(TensorInfo.name),
+						DataType,
+						MakeArrayView(TensorInfo.shape, TensorInfo.shapeLen)
+					);
+
+				Builder->AddOperatorOutput(Op, Tensor);
+			}
+		}
+
+		return Builder->End(NNEData);
+	}
+};
+
+/** Create a model optimizer */
+NNEUTILS_API TUniquePtr<IModelOptimizer> CreateModelOptimizer(ENNXInferenceFormat InputFormat, ENNXInferenceFormat OutputFormat)
+{
+	if (InputFormat == ENNXInferenceFormat::ONNX)
+	{
+		if (OutputFormat == ENNXInferenceFormat::NNXRT)
+		{
+			return MakeUnique<FModelOptimizerONNXToNNERT>();
+		}
+		else if (OutputFormat == ENNXInferenceFormat::ONNX)
+		{
+			return MakeUnique<FModelOptimizerONNXToONNX>();
+		}
+		else
+		{
+			return MakeUnique<FModelOptimizerONNXToORT>();
+		}
+	}
+
+	//TODO jira 167592: Investigate how to conditionally compile the above,
+	//removing the dependencies to ORT in runtime build (we will then
+	//need a way to run tests without ability to create onnx
+	//model at runtime). One could also make NNXUtils editor only
+	//then the interface for optimization and validation should be in a
+	//different module than the implementation.
+
+	return nullptr;
+}
+
+} // namespace UE::NNEUtils::Internal
diff --git a/Engine/Plugins/Experimental/NNE/Source/NNEUtils/Private/NNEUtilsModule.cpp b/Engine/Plugins/Experimental/NNE/Source/NNEUtils/Private/NNEUtilsModule.cpp
new file mode 100644
--- /dev/null
+++ b/Engine/Plugins/Experimental/NNE/Source/NNEUtils/Private/NNEUtilsModule.cpp
@@ -0,0 +1,18 @@
+// Copyright Epic Games, Inc. All Rights Reserved.
+
+#include "Modules/ModuleManager.h"
+
+class FNNEUtilsModule : public IModuleInterface
+{
+public:
+
+	virtual void StartupModule() override
+	{
+	}
+
+	virtual void ShutdownModule() override
+	{
+	}
+};
+
+IMPLEMENT_MODULE(FNNEUtilsModule, NNEUtils)
\ No newline at end of file
diff --git a/Engine/Plugins/Experimental/NNE/Source/NNXQA/NNXQA.Build.cs b/Engine/Plugins/Experimental/NNE/Source/NNXQA/NNXQA.Build.cs
--- a/Engine/Plugins/Experimental/NNE/Source/NNXQA/NNXQA.Build.cs
+++ b/Engine/Plugins/Experimental/NNE/Source/NNXQA/NNXQA.Build.cs
@@ -23,7 +23,7 @@ public class NNXQA : ModuleRules
 				"Json",
 				"CoreUObject",
 				"NNXCore",
-				"NNXUtils"
+				"NNEUtils"
 			}
 		);
 
diff --git a/Engine/Plugins/Experimental/NNE/Source/NNXQA/Private/NNXQAParametricTest.cpp b/Engine/Plugins/Experimental/NNE/Source/NNXQA/Private/NNXQAParametricTest.cpp
--- a/Engine/Plugins/Experimental/NNE/Source/NNXQA/Private/NNXQAParametricTest.cpp
+++ b/Engine/Plugins/Experimental/NNE/Source/NNXQA/Private/NNXQAParametricTest.cpp
@@ -8,7 +8,7 @@
 #include "NNECoreAttributeMap.h"
 #include "NNXQAUtils.h"
 #include "NNXQAJsonUtils.h"
-#include "NNXModelBuilder.h"
+#include "NNEUtilsModelBuilder.h"
 #include "HAL/IConsoleManager.h"
 #include "Interfaces/IPluginManager.h"
 #include "Misc/Paths.h"
@@ -22,6 +22,7 @@
 namespace NNX 
 {
 using FSymbolicTensorShape = UE::NNECore::FSymbolicTensorShape;
+using FTensor = UE::NNECore::Internal::FTensor;
 	
 namespace Test
 {
@@ -496,12 +497,12 @@ namespace Test
 		else
 		{
 			// Operator test, create model in memory
-			if (!CreateONNXModelForOperator(false, TestSetup.TargetName, TestSetup.Inputs, TestSetup.Outputs, TestSetup.Weights, TestSetup.WeightsData,TestSetup.AttributeMap, ONNXModel))
+			if (!UE::NNEUtils::Internal::CreateONNXModelForOperator(false, TestSetup.TargetName, TestSetup.Inputs, TestSetup.Outputs, TestSetup.Weights, TestSetup.WeightsData,TestSetup.AttributeMap, ONNXModel))
 			{
 				UE_LOG(LogNNX, Error, TEXT("Failed to create static model for test '%s'. Test ABORTED!"), *TestSetup.TargetName);
 				return false;
 			}
-			if (!CreateONNXModelForOperator(true, TestSetup.TargetName, TestSetup.Inputs, TestSetup.Outputs, TestSetup.Weights, TestSetup.WeightsData, TestSetup.AttributeMap, ONNXModelVariadic))
+			if (!UE::NNEUtils::Internal::CreateONNXModelForOperator(true, TestSetup.TargetName, TestSetup.Inputs, TestSetup.Outputs, TestSetup.Weights, TestSetup.WeightsData, TestSetup.AttributeMap, ONNXModelVariadic))
 			{
 				UE_LOG(LogNNX, Error, TEXT("Failed to create variadic model for test '%s'. Test ABORTED!"), *TestSetup.TargetName);
 				return false;
diff --git a/Engine/Plugins/Experimental/NNE/Source/NNXRuntimeCPU/NNXRuntimeCPU.Build.cs b/Engine/Plugins/Experimental/NNE/Source/NNXRuntimeCPU/NNXRuntimeCPU.Build.cs
--- a/Engine/Plugins/Experimental/NNE/Source/NNXRuntimeCPU/NNXRuntimeCPU.Build.cs
+++ b/Engine/Plugins/Experimental/NNE/Source/NNXRuntimeCPU/NNXRuntimeCPU.Build.cs
@@ -29,7 +29,7 @@ public class NNXRuntimeCPU : ModuleRules
 			new string[] {
 				//"ORTDefault",
 				"CoreUObject",
-				"NNXUtils",
+				"NNEUtils",
 				"NNX_ONNXRuntime"	// Select this for open-source-based ONNX Runtime
 			}
 		);
diff --git a/Engine/Plugins/Experimental/NNE/Source/NNXRuntimeCPU/Private/NNERuntimeCPU.cpp b/Engine/Plugins/Experimental/NNE/Source/NNXRuntimeCPU/Private/NNERuntimeCPU.cpp
--- a/Engine/Plugins/Experimental/NNE/Source/NNXRuntimeCPU/Private/NNERuntimeCPU.cpp
+++ b/Engine/Plugins/Experimental/NNE/Source/NNXRuntimeCPU/Private/NNERuntimeCPU.cpp
@@ -3,7 +3,7 @@
 #include "NNERuntimeCPU.h"
 #include "NNERuntimeCPUModel.h"
 #include "NNERuntimeCPUUtils.h"
-#include "NNXModelOptimizer.h"
+#include "NNEUtilsModelOptimizer.h"
 #include "NNXModelOptimizerInterface.h"
 #include "NNECoreAttributeMap.h"
 #include "NNECoreModelData.h"
@@ -27,7 +27,7 @@ TArray<uint8> UNNERuntimeCPUImpl::CreateModelData(FString FileType, TConstArrayV
 		return {};
 	}
 
-	TUniquePtr<IModelOptimizer> Optimizer = CreateONNXToONNXModelOptimizer();
+	TUniquePtr<IModelOptimizer> Optimizer = UE::NNEUtils::Internal::CreateONNXToONNXModelOptimizer();
 
 	FNNIModelRaw InputModel;
 	InputModel.Data = FileData;
diff --git a/Engine/Plugins/Experimental/NNE/Source/NNXRuntimeCPU/Private/NNXRuntimeCPU.cpp b/Engine/Plugins/Experimental/NNE/Source/NNXRuntimeCPU/Private/NNXRuntimeCPU.cpp
--- a/Engine/Plugins/Experimental/NNE/Source/NNXRuntimeCPU/Private/NNXRuntimeCPU.cpp
+++ b/Engine/Plugins/Experimental/NNE/Source/NNXRuntimeCPU/Private/NNXRuntimeCPU.cpp
@@ -1,7 +1,7 @@
 // Copyright Epic Games, Inc. All Rights Reserved.
 #include "NNXRuntimeCPU.h"
 #include "NNXRuntimeCPUUtils.h"
-#include "NNXModelOptimizer.h"
+#include "NNEUtilsModelOptimizer.h"
 #include "NNECoreAttributeMap.h"
 #include "NNEProfilingTimer.h"
 #include "RedirectCoutAndCerrToUeLog.h"
@@ -33,7 +33,7 @@ TArray<uint8> FRuntimeCPU::CreateModelData(FString FileType, TConstArrayView<uin
 		return {};
 	}
 
-	TUniquePtr<IModelOptimizer> Optimizer = CreateONNXToONNXModelOptimizer();
+	TUniquePtr<IModelOptimizer> Optimizer = UE::NNEUtils::Internal::CreateONNXToONNXModelOptimizer();
 
 	FNNIModelRaw InputModel;
 	InputModel.Data = FileData;
diff --git a/Engine/Plugins/Experimental/NNE/Source/NNXRuntimeORT/NNXRuntimeORT.Build.cs b/Engine/Plugins/Experimental/NNE/Source/NNXRuntimeORT/NNXRuntimeORT.Build.cs
--- a/Engine/Plugins/Experimental/NNE/Source/NNXRuntimeORT/NNXRuntimeORT.Build.cs
+++ b/Engine/Plugins/Experimental/NNE/Source/NNXRuntimeORT/NNXRuntimeORT.Build.cs
@@ -33,7 +33,7 @@ public class NNXRuntimeORT : ModuleRules
 				"ORTHelper",
 				//"ONNXRuntime",
 				"NNEOnnxruntimeEditor",
-                "NNXUtils"
+                "NNEUtils"
 			}
 		);
 
diff --git a/Engine/Plugins/Experimental/NNE/Source/NNXRuntimeORT/Private/NNXRuntimeORT.cpp b/Engine/Plugins/Experimental/NNE/Source/NNXRuntimeORT/Private/NNXRuntimeORT.cpp
--- a/Engine/Plugins/Experimental/NNE/Source/NNXRuntimeORT/Private/NNXRuntimeORT.cpp
+++ b/Engine/Plugins/Experimental/NNE/Source/NNXRuntimeORT/Private/NNXRuntimeORT.cpp
@@ -6,7 +6,7 @@
 
 #include "NNEProfilingTimer.h"
 #include "NNXRuntimeORTUtils.h"
-#include "NNXModelOptimizer.h"
+#include "NNEUtilsModelOptimizer.h"
 #include "NNECoreAttributeMap.h"
 #include "RedirectCoutAndCerrToUeLog.h"
 
@@ -75,7 +75,7 @@ TArray<uint8> FRuntimeORT::CreateModelData(FString FileType, TConstArrayView<uin
 		return {};
 	}
 
-	TUniquePtr<IModelOptimizer> Optimizer = CreateONNXToONNXModelOptimizer();
+	TUniquePtr<IModelOptimizer> Optimizer = UE::NNEUtils::Internal::CreateONNXToONNXModelOptimizer();
 
 	FNNIModelRaw InputModel;
 	InputModel.Data = FileData;
diff --git a/Engine/Plugins/Experimental/NNE/Source/NNXRuntimeRDG/NNXRuntimeRDG.Build.cs b/Engine/Plugins/Experimental/NNE/Source/NNXRuntimeRDG/NNXRuntimeRDG.Build.cs
--- a/Engine/Plugins/Experimental/NNE/Source/NNXRuntimeRDG/NNXRuntimeRDG.Build.cs
+++ b/Engine/Plugins/Experimental/NNE/Source/NNXRuntimeRDG/NNXRuntimeRDG.Build.cs
@@ -23,7 +23,7 @@ public class NNXRuntimeRDG : ModuleRules
         PrivateDependencyModuleNames.AddRange(new string[]
         {
             "NNXCore",
-			"NNXUtils",
+			"NNEUtils",
 			"NNEHlslShaders",
             "RHI"
         });
diff --git a/Engine/Plugins/Experimental/NNE/Source/NNXRuntimeRDG/Private/NNXRuntimeDML.cpp b/Engine/Plugins/Experimental/NNE/Source/NNXRuntimeRDG/Private/NNXRuntimeDML.cpp
--- a/Engine/Plugins/Experimental/NNE/Source/NNXRuntimeRDG/Private/NNXRuntimeDML.cpp
+++ b/Engine/Plugins/Experimental/NNE/Source/NNXRuntimeRDG/Private/NNXRuntimeDML.cpp
@@ -3,7 +3,7 @@
 #include "NNXRuntimeRDG.h"
 #include "NNXInferenceModel.h"
 #include "NNXRuntimeFormat.h"
-#include "NNXModelOptimizer.h"
+#include "NNEUtilsModelOptimizer.h"
 #include "NNXOperator.h"
 
 #include "NNECoreAttributeMap.h"
@@ -855,7 +855,7 @@ public:
 			return {};
 		}
 
-		TUniquePtr<IModelOptimizer> Optimizer = CreateONNXToNNXModelOptimizer();
+		TUniquePtr<IModelOptimizer> Optimizer = UE::NNEUtils::Internal::CreateONNXToNNEModelOptimizer();
 
 		FNNIModelRaw InputModel;
 		InputModel.Data = FileData;
diff --git a/Engine/Plugins/Experimental/NNE/Source/NNXRuntimeRDG/Private/NNXRuntimeHLSL.cpp b/Engine/Plugins/Experimental/NNE/Source/NNXRuntimeRDG/Private/NNXRuntimeHLSL.cpp
--- a/Engine/Plugins/Experimental/NNE/Source/NNXRuntimeRDG/Private/NNXRuntimeHLSL.cpp
+++ b/Engine/Plugins/Experimental/NNE/Source/NNXRuntimeRDG/Private/NNXRuntimeHLSL.cpp
@@ -14,7 +14,7 @@
 #include "NNECoreAttributeMap.h"
 #include "NNERuntimeRDGModelHLSL.h"
 #include "NNXInferenceModel.h"
-#include "NNXModelOptimizer.h"
+#include "NNEUtilsModelOptimizer.h"
 #include "NNXRuntimeFormat.h"
 #include "NNXRuntimeHLSLHelper.h"
 #include "NNXRuntimeRDG.h"
@@ -74,7 +74,7 @@ public:
 			return {};
 		}
 
-		TUniquePtr<NNX::IModelOptimizer> Optimizer = NNX::CreateONNXToNNXModelOptimizer();
+		TUniquePtr<NNX::IModelOptimizer> Optimizer = NNEUtils::Internal::CreateONNXToNNEModelOptimizer();
 		Optimizer->AddValidator(MakeShared<NNX::FModelValidatorHlsl>());
 
 		FNNIModelRaw InputModel;
diff --git a/Engine/Plugins/Experimental/NNE/Source/NNXRuntimeRDG/Private/NNXRuntimeRDG.cpp b/Engine/Plugins/Experimental/NNE/Source/NNXRuntimeRDG/Private/NNXRuntimeRDG.cpp
--- a/Engine/Plugins/Experimental/NNE/Source/NNXRuntimeRDG/Private/NNXRuntimeRDG.cpp
+++ b/Engine/Plugins/Experimental/NNE/Source/NNXRuntimeRDG/Private/NNXRuntimeRDG.cpp
@@ -3,7 +3,7 @@
 #include "NNXRuntimeRDG.h"
 #include "NNXInferenceModel.h"
 #include "NNXRuntimeFormat.h"
-#include "NNXModelOptimizer.h"
+#include "NNEUtilsModelOptimizer.h"
 
 #include "RenderGraphBuilder.h"
 #include "RenderGraphUtils.h"
diff --git a/Engine/Plugins/Experimental/NNE/Source/NNXUtils/NNXUtils.Build.cs b/Engine/Plugins/Experimental/NNE/Source/NNXUtils/NNXUtils.Build.cs
deleted file mode 100644
diff --git a/Engine/Plugins/Experimental/NNE/Source/NNXUtils/Private/NNXModelBuilderNNX.cpp b/Engine/Plugins/Experimental/NNE/Source/NNXUtils/Private/NNXModelBuilderNNX.cpp
deleted file mode 100644
diff --git a/Engine/Plugins/Experimental/NNE/Source/NNXUtils/Private/NNXModelBuilderONNX.cpp b/Engine/Plugins/Experimental/NNE/Source/NNXUtils/Private/NNXModelBuilderONNX.cpp
deleted file mode 100644
diff --git a/Engine/Plugins/Experimental/NNE/Source/NNXUtils/Private/NNXModelOptimizer.cpp b/Engine/Plugins/Experimental/NNE/Source/NNXUtils/Private/NNXModelOptimizer.cpp
deleted file mode 100644
diff --git a/Engine/Plugins/Experimental/NNE/Source/NNXUtils/Private/NNXUtilsModule.cpp b/Engine/Plugins/Experimental/NNE/Source/NNXUtils/Private/NNXUtilsModule.cpp
deleted file mode 100644
diff --git a/Engine/Plugins/Experimental/NNE/Source/NNXUtils/Public/NNXModelBuilder.h b/Engine/Plugins/Experimental/NNE/Source/NNXUtils/Public/NNXModelBuilder.h
deleted file mode 100644
diff --git a/Engine/Plugins/Experimental/NNE/Source/NNXUtils/Public/NNXModelOptimizer.h b/Engine/Plugins/Experimental/NNE/Source/NNXUtils/Public/NNXModelOptimizer.h
deleted file mode 100644
