NNE - Small cleanup following review + fix slow test (NNERuntimeCPU was testing against itself)

diff --git a/Engine/Plugins/Experimental/NNE/Source/NNEHlslShaders/Internal/NNXShaderParameters.h b/Engine/Plugins/Experimental/NNE/Source/NNEHlslShaders/Internal/NNXShaderParameters.h
deleted file mode 100644
diff --git a/Engine/Plugins/Experimental/NNE/Source/NNXQA/Private/NNXQAUtils.cpp b/Engine/Plugins/Experimental/NNE/Source/NNXQA/Private/NNXQAUtils.cpp
--- a/Engine/Plugins/Experimental/NNE/Source/NNXQA/Private/NNXQAUtils.cpp
+++ b/Engine/Plugins/Experimental/NNE/Source/NNXQA/Private/NNXQAUtils.cpp
@@ -687,7 +687,7 @@ namespace Test
 		for (auto Runtime : Runtimes)
 		{
 			const FString& RuntimeName = Runtime->GetRuntimeName();
-			if (RuntimeName == RefName)
+			if ((RuntimeName == RefName) || RuntimeName == TEXT("NNXRuntimeCPU"))
 			{
 				continue;
 			}
diff --git a/Engine/Plugins/Experimental/NNE/Source/NNXRuntimeRDG/Private/NNXRuntimeHLSLHelper.h b/Engine/Plugins/Experimental/NNE/Source/NNXRuntimeRDG/Private/NNXRuntimeHLSLHelper.h
--- a/Engine/Plugins/Experimental/NNE/Source/NNXRuntimeRDG/Private/NNXRuntimeHLSLHelper.h
+++ b/Engine/Plugins/Experimental/NNE/Source/NNXRuntimeRDG/Private/NNXRuntimeHLSLHelper.h
@@ -3,12 +3,14 @@
 #pragma once
 
 #include "Containers/StaticArray.h"
-#include "NNXShaderParameters.h"
 
 namespace UE::NNECore::Internal { class FTensor; }
 
+#define NXRT_TENSORSTRIDEINFO_MAX_NUM_DIMENSIONS 8
+
 namespace NNX
 {
+
 	void FillTensorSizeShaderParameters(const UE::NNECore::Internal::FTensor& Tensor, TStaticArray<FUintVector4, NXRT_TENSORSTRIDEINFO_MAX_NUM_DIMENSIONS, 16U>& OutShaderParam, int32 Idx);
 	void FillTensorStrideShaderParameters(const UE::NNECore::Internal::FTensor& Tensor, TStaticArray<FUintVector4, NXRT_TENSORSTRIDEINFO_MAX_NUM_DIMENSIONS, 16U>& OutShaderParam, int32 Idx, int32 TargetNumdimensionForBroadcast = -1);
 	void FillTensorStrideForBroadcastShaderParameters(const UE::NNECore::Internal::FTensor& Tensor, int32 OutputNumdimension, TStaticArray<FUintVector4, NXRT_TENSORSTRIDEINFO_MAX_NUM_DIMENSIONS, 16U>& OutShaderParam, int32 Idx);
