NNE : Minimal API - NNERuntimeORT

Problem:
* NNE API surface is large atm potentially creating a support burden and harming adoption for 5.2.

Solution:
* This CL is targeting NNERuntimeORT to use the minimal API.
* cleanup/rename module/file/classes to follow naming convention.
* remove ORTCuda

Test:
* All tests green under win64 editor + style transfer demo in editor

diff --git a/Engine/Plugins/Experimental/NNE/NNE.uplugin b/Engine/Plugins/Experimental/NNE/NNE.uplugin
--- a/Engine/Plugins/Experimental/NNE/NNE.uplugin
+++ b/Engine/Plugins/Experimental/NNE/NNE.uplugin
@@ -62,7 +62,7 @@
 			"PlatformAllowList": [ "Win64", "Linux", "Mac" ]
 		},
 		{
-			"Name": "NNXRuntimeORT",
+			"Name": "NNERuntimeORT",
 			"Type": "EditorAndProgram",
 			"LoadingPhase": "Default",
 			"PlatformAllowList": [ "Win64", "Linux", "Mac" ]
diff --git a/Engine/Plugins/Experimental/NNE/Source/NNERuntimeORT/NNERuntimeORT.Build.cs b/Engine/Plugins/Experimental/NNE/Source/NNERuntimeORT/NNERuntimeORT.Build.cs
new file mode 100644
--- /dev/null
+++ b/Engine/Plugins/Experimental/NNE/Source/NNERuntimeORT/NNERuntimeORT.Build.cs
@@ -0,0 +1,60 @@
+// Copyright Epic Games, Inc. All Rights Reserved.
+
+using UnrealBuildTool;
+using System.IO;
+
+public class NNERuntimeORT : ModuleRules
+{
+	public NNERuntimeORT( ReadOnlyTargetRules Target ) : base( Target )
+	{
+		ShortName = "NNERuntimeORT"; // Shorten to avoid path-too-long errors
+		PCHUsage = ModuleRules.PCHUsageMode.UseExplicitOrSharedPCHs;
+
+		PublicIncludePaths.AddRange(
+			new string[] {
+				System.IO.Path.Combine(ModuleDirectory, "..")
+			}
+		);
+
+		PublicDependencyModuleNames.AddRange
+			(
+			new string[] {
+				"Core",
+				"Engine",
+				"NNXCore",
+				"NNEProfiling",
+			}
+		);
+
+		PrivateDependencyModuleNames.AddRange
+			(
+			new string[] {
+				// ORT-related
+				"CoreUObject",
+				"ORTHelper",
+				//"ONNXRuntime",
+				"NNEOnnxruntimeEditor",
+                "NNEUtils"
+			}
+		);
+
+		// Win64-only
+		if (Target.Platform == UnrealTargetPlatform.Win64)
+		{
+			PublicDefinitions.Add("USE_DML");
+			PublicDefinitions.Add("PLATFORM_WIN64");
+
+			PrivateDependencyModuleNames.AddRange(new string[]
+			{
+				"RHI",
+			});
+
+			PrivateDependencyModuleNames.Add("D3D12RHI");
+			PrivateDependencyModuleNames.Add("DirectMLDefault");
+
+			AddEngineThirdPartyPrivateStaticDependencies(Target, "DX12");
+			AddEngineThirdPartyPrivateStaticDependencies(Target, "DirectMLDefault");
+
+		}
+	}
+}
diff --git a/Engine/Plugins/Experimental/NNE/Source/NNERuntimeORT/Private/NNERuntimeORT.cpp b/Engine/Plugins/Experimental/NNE/Source/NNERuntimeORT/Private/NNERuntimeORT.cpp
new file mode 100644
--- /dev/null
+++ b/Engine/Plugins/Experimental/NNE/Source/NNERuntimeORT/Private/NNERuntimeORT.cpp
@@ -0,0 +1,161 @@
+// Copyright Epic Games, Inc. All Rights Reserved.
+
+#include "NNERuntimeORT.h"
+#include "NNERuntimeORTModel.h"
+#include "NNERuntimeORTUtils.h"
+#include "NNEUtilsModelOptimizer.h"
+#include "NNXModelOptimizerInterface.h"
+#include "NNECoreAttributeMap.h"
+#include "NNECoreModelData.h"
+#include "NNEProfilingTimer.h"
+#include "RedirectCoutAndCerrToUeLog.h"
+
+#include "NNXThirdPartyWarningDisabler.h"
+NNX_THIRD_PARTY_INCLUDES_START
+#undef check
+#undef TEXT
+#include "core/session/onnxruntime_cxx_api.h"
+NNX_THIRD_PARTY_INCLUDES_END
+
+
+using namespace NNX;
+
+FGuid UNNERuntimeORTCpuImpl::GUID = FGuid((int32)'O', (int32)'N', (int32)'N', (int32)'X');
+int32 UNNERuntimeORTCpuImpl::Version = 0x00000001;
+
+bool UNNERuntimeORTCpuImpl::CanCreateModelData(FString FileType, TConstArrayView<uint8> FileData) const
+{
+	return FileType.Compare("onnx", ESearchCase::IgnoreCase) == 0;
+}
+
+TArray<uint8> UNNERuntimeORTCpuImpl::CreateModelData(FString FileType, TConstArrayView<uint8> FileData)
+{
+	if (!CanCreateModelData(FileType, FileData))
+	{
+		return {};
+	}
+
+	TUniquePtr<IModelOptimizer> Optimizer = UE::NNEUtils::Internal::CreateONNXToONNXModelOptimizer();
+
+	FNNIModelRaw InputModel;
+	InputModel.Data = FileData;
+	InputModel.Format = ENNXInferenceFormat::ONNX;
+	FNNIModelRaw OutputModel;
+	FOptimizerOptionsMap Options;
+	if (!Optimizer->Optimize(InputModel, OutputModel, Options))
+	{
+		return {};
+	}
+
+	int32 GuidSize = sizeof(UNNERuntimeORTCpuImpl::GUID);
+	int32 VersionSize = sizeof(UNNERuntimeORTCpuImpl::Version);
+	TArray<uint8> Result;
+	FMemoryWriter Writer(Result);
+	Writer << UNNERuntimeORTCpuImpl::GUID;
+	Writer << UNNERuntimeORTCpuImpl::Version;
+	Writer.Serialize(OutputModel.Data.GetData(), OutputModel.Data.Num());
+	return Result;
+}
+
+bool UNNERuntimeORTCpuImpl::CanCreateModelCPU(TObjectPtr<UNNEModelData> ModelData) const
+{
+	check(ModelData != nullptr);
+
+	int32 GuidSize = sizeof(UNNERuntimeORTCpuImpl::GUID);
+	int32 VersionSize = sizeof(UNNERuntimeORTCpuImpl::Version);
+	TConstArrayView<uint8> Data = ModelData->GetModelData(GetRuntimeName());
+
+	if (Data.Num() <= GuidSize + VersionSize)
+	{
+		return false;
+	}
+
+	bool bResult = FGenericPlatformMemory::Memcmp(&(Data[0]), &(UNNERuntimeORTCpuImpl::GUID), GuidSize) == 0;
+	bResult &= FGenericPlatformMemory::Memcmp(&(Data[GuidSize]), &(UNNERuntimeORTCpuImpl::Version), VersionSize) == 0;
+	return bResult;
+}
+
+void UNNERuntimeORTCpuImpl::Init()
+{
+	check(!ORTEnvironment.IsValid());
+	ORTEnvironment = MakeUnique<Ort::Env>();
+}
+
+TUniquePtr<UE::NNECore::IModelCPU> UNNERuntimeORTCpuImpl::CreateModelCPU(TObjectPtr<UNNEModelData> ModelData)
+{
+	check(ModelData != nullptr);
+	check(ORTEnvironment.IsValid());
+
+	if (!CanCreateModelCPU(ModelData))
+	{
+		return TUniquePtr<UE::NNECore::IModelCPU>();
+	}
+
+	const UE::NNERuntimeORT::Private::FRuntimeConf InConf;
+	UE::NNERuntimeORT::Private::FModelORTCpu* Model = new UE::NNERuntimeORT::Private::FModelORTCpu(ORTEnvironment.Get(), InConf);
+	TConstArrayView<uint8> Data = ModelData->GetModelData(GetRuntimeName());
+
+	if (!Model->Init(Data))
+	{
+		delete Model;
+		return TUniquePtr<UE::NNECore::IModelCPU>();
+	}
+	UE::NNECore::IModelCPU* IModel = static_cast<UE::NNECore::IModelCPU*>(Model);
+	return TUniquePtr<UE::NNECore::IModelCPU>(IModel);
+}
+
+#if PLATFORM_WINDOWS
+bool UNNERuntimeORTDmlImpl::CanCreateModelCPU(TObjectPtr<UNNEModelData> ModelData) const
+{
+	check(ModelData != nullptr);
+
+	int32 GuidSize = sizeof(UNNERuntimeORTCpuImpl::GUID);
+	int32 VersionSize = sizeof(UNNERuntimeORTCpuImpl::Version);
+	TConstArrayView<uint8> Data = ModelData->GetModelData(GetRuntimeName());
+
+	if (Data.Num() <= GuidSize + VersionSize)
+	{
+		return false;
+	}
+
+	bool bResult = FGenericPlatformMemory::Memcmp(&(Data[0]), &(UNNERuntimeORTCpuImpl::GUID), GuidSize) == 0;
+	bResult &= FGenericPlatformMemory::Memcmp(&(Data[GuidSize]), &(UNNERuntimeORTCpuImpl::Version), VersionSize) == 0;
+	return bResult;
+}
+
+TUniquePtr<UE::NNECore::IModelCPU> UNNERuntimeORTDmlImpl::CreateModelCPU(TObjectPtr<UNNEModelData> ModelData)
+{
+	check(ModelData != nullptr);
+	check(ORTEnvironment.IsValid());
+
+	if (!CanCreateModelCPU(ModelData))
+	{
+		return TUniquePtr<UE::NNECore::IModelCPU>();
+	}
+
+	const UE::NNERuntimeORT::Private::FRuntimeConf InConf;
+	UE::NNERuntimeORT::Private::FModelORTDml* Model = new UE::NNERuntimeORT::Private::FModelORTDml(ORTEnvironment.Get(), InConf);
+	TConstArrayView<uint8> Data = ModelData->GetModelData(GetRuntimeName());
+
+	if (!Model->Init(Data))
+	{
+		delete Model;
+		return TUniquePtr<UE::NNECore::IModelCPU>();
+	}
+	UE::NNECore::IModelCPU* IModel = static_cast<UE::NNECore::IModelCPU*>(Model);
+	return TUniquePtr<UE::NNECore::IModelCPU>(IModel);
+}
+
+#else // PLATFORM_WINDOWS
+
+bool UNNERuntimeORTDmlImpl::CanCreateModelCPU(TObjectPtr<UNNEModelData> ModelData) const
+{
+	return false;
+}
+
+TUniquePtr<UE::NNECore::IModelCPU> UNNERuntimeORTDmlImpl::CreateModelCPU(TObjectPtr<UNNEModelData> ModelData)
+{
+	return TUniquePtr<UE::NNECore::IModelCPU>();
+}
+
+#endif // PLATFORM_WINDOWS
\ No newline at end of file
diff --git a/Engine/Plugins/Experimental/NNE/Source/NNERuntimeORT/Private/NNERuntimeORT.h b/Engine/Plugins/Experimental/NNE/Source/NNERuntimeORT/Private/NNERuntimeORT.h
new file mode 100644
--- /dev/null
+++ b/Engine/Plugins/Experimental/NNE/Source/NNERuntimeORT/Private/NNERuntimeORT.h
@@ -0,0 +1,50 @@
+// Copyright Epic Games, Inc. All Rights Reserved.
+
+#pragma once
+
+#include "CoreMinimal.h"
+#include "NNECoreRuntime.h"
+#include "NNECoreRuntimeCPU.h"
+#include "Templates/UniquePtr.h"
+#include "UObject/Class.h"
+#include "UObject/Object.h"
+#include "UObject/UObjectBaseUtility.h"
+
+#include "NNERuntimeORT.generated.h"
+
+namespace Ort { struct Env; }
+
+UCLASS()
+class UNNERuntimeORTCpuImpl : public UObject, public INNERuntime, public INNERuntimeCPU
+{
+	GENERATED_BODY()
+
+public:
+	static FGuid GUID;
+	static int32 Version;
+
+	TUniquePtr<Ort::Env> ORTEnvironment;
+	UNNERuntimeORTCpuImpl() {};
+	virtual ~UNNERuntimeORTCpuImpl() {}
+
+	void Init();
+
+	virtual FString GetRuntimeName() const override { return TEXT("NNERuntimeORTCpuEditor"); };
+	virtual bool IsPlatformSupported(const ITargetPlatform* TargetPlatform) const override { return true; };
+
+	virtual bool CanCreateModelData(FString FileType, TConstArrayView<uint8> FileData) const override;
+	virtual TArray<uint8> CreateModelData(FString FileType, TConstArrayView<uint8> FileData) override;
+
+	virtual bool CanCreateModelCPU(TObjectPtr<UNNEModelData> ModelData) const override;
+	virtual TUniquePtr<UE::NNECore::IModelCPU> CreateModelCPU(TObjectPtr<UNNEModelData> ModelData) override;
+};
+
+UCLASS()
+class UNNERuntimeORTDmlImpl : public UNNERuntimeORTCpuImpl
+{
+	GENERATED_BODY()
+
+	virtual FString GetRuntimeName() const override { return TEXT("NNERuntimeORTDmlEditor"); };
+	virtual bool CanCreateModelCPU(TObjectPtr<UNNEModelData> ModelData) const override;
+	virtual TUniquePtr<UE::NNECore::IModelCPU> CreateModelCPU(TObjectPtr<UNNEModelData> ModelData) override;
+};
\ No newline at end of file
diff --git a/Engine/Plugins/Experimental/NNE/Source/NNERuntimeORT/Private/NNERuntimeORTModel.cpp b/Engine/Plugins/Experimental/NNE/Source/NNERuntimeORT/Private/NNERuntimeORTModel.cpp
new file mode 100644
--- /dev/null
+++ b/Engine/Plugins/Experimental/NNE/Source/NNERuntimeORT/Private/NNERuntimeORTModel.cpp
@@ -0,0 +1,414 @@
+// Copyright Epic Games, Inc. All Rights Reserved.
+#include "NNERuntimeORTModel.h"
+#include "NNERuntimeORT.h"
+#include "NNERuntimeORTUtils.h"
+#include "NNEProfilingTimer.h"
+#include "RedirectCoutAndCerrToUeLog.h"
+
+#if PLATFORM_WINDOWS
+
+#include "Windows/AllowWindowsPlatformTypes.h"
+#include <unknwn.h>
+#include "Microsoft/COMPointer.h"
+#include "Windows/HideWindowsPlatformTypes.h"
+
+#include "ID3D12DynamicRHI.h"
+#include "D3D12RHIBridge.h"
+#include "DirectML.h"
+#endif
+
+NNX_THIRD_PARTY_INCLUDES_START
+#undef check
+#undef TEXT
+
+#include "core/session/onnxruntime_cxx_api.h"
+
+#include "core/providers/cpu/cpu_provider_factory.h"
+
+#if PLATFORM_WINDOWS
+#include "Windows/MinWindows.h"
+
+#ifdef USE_DML
+#include "core/providers/dml/dml_provider_factory.h"
+#endif
+#endif
+
+NNX_THIRD_PARTY_INCLUDES_END
+
+namespace UE::NNERuntimeORT::Private
+{
+
+	FModelORT::FModelORT() :
+		bIsLoaded(false),
+		bHasRun(false)
+	{ }
+
+	FModelORT::FModelORT(
+		Ort::Env* InORTEnvironment,
+		const FRuntimeConf& InRuntimeConf) :
+		bIsLoaded(false),
+		bHasRun(false),
+		RuntimeConf(InRuntimeConf),
+		ORTEnvironment(InORTEnvironment)
+	{ }
+
+	bool FModelORT::Init(TConstArrayView<uint8> ModelData)
+	{
+		DECLARE_SCOPE_CYCLE_COUNTER(TEXT("FModelORT_Init"), STAT_FModelORT_Init, STATGROUP_NNE);
+		
+		// Get the header size
+		int32 GuidSize = sizeof(UNNERuntimeORTCpuImpl::GUID);
+		int32 VersionSize = sizeof(UNNERuntimeORTCpuImpl::Version);
+
+		// Clean previous networks
+		bIsLoaded = false;
+		TConstArrayView<uint8> ModelBuffer = TConstArrayView<uint8>(&(ModelData.GetData()[GuidSize + VersionSize]), ModelData.Num() - GuidSize - VersionSize);
+
+		// Checking Inference Model 
+		{
+			if (ModelBuffer.Num() == 0)
+			{
+				UE_LOG(LogNNE, Warning, TEXT("FModelORT::Load(): Input model data is empty."));
+				return false;
+			}
+
+		}
+
+#if WITH_EDITOR
+		try
+#endif //WITH_EDITOR
+		{
+			const FRedirectCoutAndCerrToUeLog RedirectCoutAndCerrToUeLog;
+			if (!InitializedAndConfigureMembers())
+			{
+				UE_LOG(LogNNE, Warning, TEXT("Load(): InitializedAndConfigureMembers failed."));
+				return false;
+			}
+
+			{
+				DECLARE_SCOPE_CYCLE_COUNTER(TEXT("FModelORT_Init_CreateORTSession"), STAT_FModelORT_Init_CreateORTSession, STATGROUP_NNE);
+				// Read model from InferenceModel
+				Session = MakeUnique<Ort::Session>(*ORTEnvironment, ModelBuffer.GetData(), ModelBuffer.Num(), *SessionOptions);
+			}
+
+			if (!ConfigureTensors(true))
+			{
+				UE_LOG(LogNNE, Warning, TEXT("Load(): Failed to configure Inputs tensors."));
+				return false;
+			}
+
+			if (!ConfigureTensors(false))
+			{
+				UE_LOG(LogNNE, Warning, TEXT("Load(): Failed to configure Outputs tensors."));
+				return false;
+			}
+		}
+#if WITH_EDITOR
+		catch (const std::exception& Exception)
+		{
+			UE_LOG(LogNNE, Error, TEXT("%s"), UTF8_TO_TCHAR(Exception.what()));
+			return false;
+		}
+#endif //WITH_EDITOR
+
+		bIsLoaded = true;
+
+		// Reset Stats
+		ResetStats();
+
+		return IsLoaded();
+	}
+
+	bool FModelORT::IsLoaded() const
+	{
+		return bIsLoaded;
+	}
+
+	bool FModelORT::InitializedAndConfigureMembers()
+	{
+		// Initialize 
+		// Setting up ORT
+		Allocator = MakeUnique<Ort::AllocatorWithDefaultOptions>();
+		AllocatorInfo = MakeUnique<Ort::MemoryInfo>(Ort::MemoryInfo::CreateCpu(OrtDeviceAllocator, OrtMemTypeCPU));
+
+		// Configure Session
+		SessionOptions = MakeUnique<Ort::SessionOptions>();
+
+		// Configure number threads
+		SessionOptions->SetIntraOpNumThreads(RuntimeConf.NumberOfThreads);
+
+		// ORT Setting optimizations to the fastest possible
+		SessionOptions->SetGraphOptimizationLevel(RuntimeConf.OptimizationLevel); // ORT_ENABLE_ALL, ORT_ENABLE_EXTENDED, ORT_ENABLE_BASIC, ORT_DISABLE_ALL
+
+		return true;
+	}
+
+
+	bool FModelORT::ConfigureTensors(const bool InIsInput)
+	{
+		const bool bIsInput = InIsInput;
+
+		const uint32 NumberTensors = bIsInput ? Session->GetInputCount() : Session->GetOutputCount();
+
+		TArray<NNECore::FTensorDesc>& SymbolicTensorDescs = bIsInput ? InputSymbolicTensors : OutputSymbolicTensors;
+		TArray<ONNXTensorElementDataType>& TensorsORTType = bIsInput ? InputTensorsORTType : OutputTensorsORTType;
+		TArray<const char*>& TensorNames = bIsInput ? InputTensorNames : OutputTensorNames;
+
+		for (uint32 TensorIndex = 0; TensorIndex < NumberTensors; ++TensorIndex)
+		{
+			// Get Tensor name
+			const char* CurTensorName = bIsInput ? Session->GetInputName(TensorIndex, *Allocator) : Session->GetOutputName(TensorIndex, *Allocator);
+			TensorNames.Emplace(CurTensorName);
+
+			// Get node type
+			Ort::TypeInfo CurrentTypeInfo = bIsInput ? Session->GetInputTypeInfo(TensorIndex) : Session->GetOutputTypeInfo(TensorIndex);
+			Ort::TensorTypeAndShapeInfo CurrentTensorInfo = CurrentTypeInfo.GetTensorTypeAndShapeInfo();
+			const ONNXTensorElementDataType ONNXTensorElementDataTypeEnum = CurrentTensorInfo.GetElementType();
+			CurrentTypeInfo.release();
+
+			TensorsORTType.Emplace(ONNXTensorElementDataTypeEnum);
+
+			std::pair<ENNETensorDataType, uint64> TypeAndSize = TranslateTensorTypeORTToNNE(ONNXTensorElementDataTypeEnum);
+
+			TArray<int32> ShapeData;
+			ShapeData.Reserve(CurrentTensorInfo.GetShape().size());
+			for (const int64_t CurrentTensorSize : CurrentTensorInfo.GetShape())
+			{
+				ShapeData.Add((int32)CurrentTensorSize);
+			}
+
+			NNECore::FSymbolicTensorShape Shape = NNECore::FSymbolicTensorShape::Make(ShapeData);
+			NNECore::FTensorDesc SymbolicTensorDesc = NNECore::FTensorDesc::Make(FString(CurTensorName), Shape, TypeAndSize.first);
+
+			check(SymbolicTensorDesc.GetElemByteSize() == TypeAndSize.second);
+			SymbolicTensorDescs.Emplace(SymbolicTensorDesc);
+		}
+
+		return true;
+	}
+	
+	int FModelORT::SetInputTensorShapes(TConstArrayView<NNECore::FTensorShape> InInputShapes)
+	{
+		InputTensors.Empty();
+		OutputTensors.Empty();
+		OutputTensorShapes.Empty();
+
+		// Verify input shape are valid for the model and set InputTensorShapes
+		if (FModelBase<IModelCPU>::SetInputTensorShapes(InInputShapes) != 0)
+		{
+			return -1;
+		}
+
+		// Setup concrete input tensor
+		for (int i = 0; i < InputSymbolicTensors.Num(); ++i)
+		{
+			NNECore::Internal::FTensor Tensor = NNECore::Internal::FTensor::Make(InputSymbolicTensors[i].GetName(), InInputShapes[i], InputSymbolicTensors[i].GetDataType());
+			InputTensors.Emplace(Tensor);
+		}
+
+		// Here model optimization could be done now that we know the input shapes, for some models
+		// this would allow to resolve output shapes here rather than during inference.
+
+		// Setup concrete output shapes only if all model output shapes are concretes, otherwise it will be set during Run()
+		for (NNECore::FTensorDesc SymbolicTensorDesc : OutputSymbolicTensors)
+		{
+			if (SymbolicTensorDesc.GetShape().IsConcrete())
+			{
+				NNECore::Internal::FTensor Tensor = NNECore::Internal::FTensor::MakeFromSymbolicDesc(SymbolicTensorDesc);
+				OutputTensors.Emplace(Tensor);
+				OutputTensorShapes.Emplace(Tensor.GetShape());
+			}
+		}
+		if (OutputTensors.Num() != OutputSymbolicTensors.Num())
+		{
+			OutputTensors.Empty();
+			OutputTensorShapes.Empty();
+		}
+
+		return 0;
+	}
+
+	int FModelORT::RunSync(TConstArrayView<NNECore::FTensorBindingCPU> InInputBindings, TConstArrayView<NNECore::FTensorBindingCPU> InOutputBindings)
+	{
+		DECLARE_SCOPE_CYCLE_COUNTER(TEXT("FModelORT_Run"), STAT_FModelORT_Run, STATGROUP_NNE);
+
+		// Sanity check
+		if (!bIsLoaded)
+		{
+			UE_LOG(LogNNE, Warning, TEXT("FModelORT::Run(): Call FModelORT::Load() to load a model first."));
+			return -1;
+		}
+
+		// Verify the model inputs were prepared
+		if (InputTensorShapes.Num() == 0)
+		{
+			UE_LOG(LogNNE, Error, TEXT("Run(): Input shapes are not set, please call SetInputTensorShapes."));
+			return -1;
+		}
+
+		NNEProfiling::Internal::FTimer RunTimer;
+		RunTimer.Tic();
+
+		if (!bHasRun)
+		{
+			bHasRun = true;
+		}
+
+#if WITH_EDITOR
+		try
+#endif //WITH_EDITOR
+		{
+			TArray<Ort::Value> InputOrtTensors;
+			BindTensorsToORT(InInputBindings, InputTensors, InputTensorsORTType, AllocatorInfo.Get(), InputOrtTensors);
+
+			if (!OutputTensors.IsEmpty())
+			{
+				// If output shapes are known we can directly map preallocated output buffers
+				TArray<Ort::Value> OutputOrtTensors;
+				BindTensorsToORT(InOutputBindings, OutputTensors, OutputTensorsORTType, AllocatorInfo.Get(), OutputOrtTensors);
+
+				Session->Run(Ort::RunOptions{ nullptr },
+					InputTensorNames.GetData(), &InputOrtTensors[0], InputTensorNames.Num(),
+					OutputTensorNames.GetData(), &OutputOrtTensors[0], OutputTensorNames.Num());
+			}
+			else
+			{
+				TArray<Ort::Value> OutputOrtTensors;
+				for (int i = 0; i < InOutputBindings.Num(); ++i)
+				{
+					OutputOrtTensors.Emplace(nullptr);
+				}
+
+				Session->Run(Ort::RunOptions{ nullptr },
+					InputTensorNames.GetData(), &InputOrtTensors[0], InputTensorNames.Num(),
+					OutputTensorNames.GetData(), &OutputOrtTensors[0], OutputTensorNames.Num());
+
+				// Output shapes were resolved during inference: Copy the data back to bindings and expose output tensor shapes
+				CopyFromORTToBindings(OutputOrtTensors, InOutputBindings, OutputSymbolicTensors, OutputTensors);
+				check(OutputTensorShapes.IsEmpty());
+				for (int i = 0; i < OutputTensors.Num(); ++i)
+				{
+					OutputTensorShapes.Emplace(OutputTensors[i].GetShape());
+				}
+			}
+		}
+#if WITH_EDITOR
+		catch (const std::exception& Exception)
+		{
+			UE_LOG(LogNNE, Error, TEXT("%s"), UTF8_TO_TCHAR(Exception.what()));
+		}
+#endif //WITH_EDITOR
+
+		RunStatisticsEstimator.StoreSample(RunTimer.Toc());
+
+		return 0;
+	}
+
+	float FModelORT::GetLastRunTimeMSec() const
+	{
+		return RunStatisticsEstimator.GetLastSample();
+	}
+
+	NNEProfiling::Internal::FStatistics FModelORT::GetRunStatistics() const
+	{
+		return RunStatisticsEstimator.GetStats();
+	}
+
+	NNEProfiling::Internal::FStatistics FModelORT::GetInputMemoryTransferStats() const
+	{
+		return InputTransferStatisticsEstimator.GetStats();
+	}
+
+	void FModelORT::ResetStats()
+	{
+		RunStatisticsEstimator.ResetStats();
+		InputTransferStatisticsEstimator.ResetStats();
+	}
+
+	bool FModelORTCpu::InitializedAndConfigureMembers()
+	{
+		if (!FModelORT::InitializedAndConfigureMembers())
+		{
+			return false;
+		}
+
+		SessionOptions->EnableCpuMemArena();
+
+		return true;
+	}
+
+#if PLATFORM_WINDOWS
+	bool FModelORTDml::InitializedAndConfigureMembers()
+	{
+		if (!FModelORT::InitializedAndConfigureMembers())
+		{
+			return false;
+		}
+
+		SessionOptions->DisableCpuMemArena();
+
+		HRESULT res;
+
+		// In order to use DirectML we need D3D12
+		ID3D12DynamicRHI* RHI = nullptr;
+
+		if (GDynamicRHI && GDynamicRHI->GetInterfaceType() == ERHIInterfaceType::D3D12)
+		{
+			RHI = static_cast<ID3D12DynamicRHI*>(GDynamicRHI);
+
+			if (!RHI)
+			{
+				UE_LOG(LogNNE, Warning, TEXT("Error:%s RHI is not supported by DirectML"), GDynamicRHI->GetName());
+				return false;
+			}
+		}
+		else
+		{
+			if (GDynamicRHI)
+			{
+				UE_LOG(LogNNE, Warning, TEXT("Error:%s RHI is not supported by DirectML"), GDynamicRHI->GetName());
+				return false;
+			}
+			else
+			{
+				UE_LOG(LogNNE, Warning, TEXT("Error:No RHI found"));
+				return false;
+			}
+		}
+
+		int DeviceIndex = 0;
+		ID3D12Device* D3D12Device = RHI->RHIGetDevice(DeviceIndex);
+
+		DML_CREATE_DEVICE_FLAGS DmlCreateFlags = DML_CREATE_DEVICE_FLAG_NONE;
+
+		// Set debugging flags
+		if (RHI->IsD3DDebugEnabled())
+		{
+			DmlCreateFlags |= DML_CREATE_DEVICE_FLAG_DEBUG;
+		}
+
+		IDMLDevice* DmlDevice;
+
+		res = DMLCreateDevice(D3D12Device, DmlCreateFlags, IID_PPV_ARGS(&DmlDevice));
+		if (!DmlDevice)
+		{
+			UE_LOG(LogNNE, Warning, TEXT("Failed to create DML device"));
+			return false;
+		}
+
+		ID3D12CommandQueue* CmdQ = RHI->RHIGetCommandQueue();
+
+		//OrtSessionOptionsAppendExecutionProvider_DML(*SessionOptions.Get(), ORTConfiguration.DeviceId);
+
+		OrtStatusPtr Status = OrtSessionOptionsAppendExecutionProviderEx_DML(*SessionOptions.Get(), DmlDevice, CmdQ);
+		if (Status)
+		{
+			UE_LOG(LogNNE, Warning, TEXT("Failed to initialize session options for ORT Dml EP: %s"), ANSI_TO_TCHAR(Ort::GetApi().GetErrorMessage(Status)));
+			return false;
+		}
+
+		return true;
+	}
+#endif
+	
+} // namespace UE::NNERuntimeORT::Private
\ No newline at end of file
diff --git a/Engine/Plugins/Experimental/NNE/Source/NNERuntimeORT/Private/NNERuntimeORTModel.h b/Engine/Plugins/Experimental/NNE/Source/NNERuntimeORT/Private/NNERuntimeORTModel.h
new file mode 100644
--- /dev/null
+++ b/Engine/Plugins/Experimental/NNE/Source/NNERuntimeORT/Private/NNERuntimeORTModel.h
@@ -0,0 +1,103 @@
+// Copyright Epic Games, Inc. All Rights Reserved.
+
+#pragma once
+
+#include "CoreMinimal.h"
+#include "NNECoreModelBase.h"
+#include "NNECoreRuntimeCPU.h"
+#include "NNECoreTensor.h"
+#include "NNECoreTypes.h"
+#include "NNEProfilingStatistics.h"
+
+#include "NNXThirdPartyWarningDisabler.h"
+NNX_THIRD_PARTY_INCLUDES_START
+#undef check
+#undef TEXT
+#include "core/session/onnxruntime_cxx_api.h"
+NNX_THIRD_PARTY_INCLUDES_END
+
+namespace UE::NNERuntimeORT::Private
+{
+	struct FRuntimeConf
+	{
+		uint32 NumberOfThreads = 2;
+		GraphOptimizationLevel OptimizationLevel = GraphOptimizationLevel::ORT_ENABLE_ALL;
+		EThreadPriority ThreadPriority = EThreadPriority::TPri_Normal;
+	};
+
+	class FModelORT : public NNECore::Internal::FModelBase<NNECore::IModelCPU>
+	{
+
+	public:
+		FModelORT();
+		FModelORT(Ort::Env* InORTEnvironment, const FRuntimeConf& InRuntimeConf);
+		virtual ~FModelORT() = default;
+
+		virtual int SetInputTensorShapes(TConstArrayView<NNECore::FTensorShape> InInputShapes) override;
+		virtual int RunSync(TConstArrayView<NNECore::FTensorBindingCPU> InInputBindings, TConstArrayView<NNECore::FTensorBindingCPU> InOutputBindings) override;
+
+		bool Init(TConstArrayView<uint8> ModelData);
+		bool IsLoaded() const;
+		float GetLastRunTimeMSec() const;
+		NNEProfiling::Internal::FStatistics GetRunStatistics() const;
+		NNEProfiling::Internal::FStatistics GetInputMemoryTransferStats() const;
+		void ResetStats();
+
+	protected:
+
+		bool bIsLoaded = false;
+		bool bHasRun = false;
+
+		FRuntimeConf RuntimeConf;
+
+		/** ORT-related variables */
+		Ort::Env* ORTEnvironment;
+		TUniquePtr<Ort::Session> Session;
+		TUniquePtr<Ort::AllocatorWithDefaultOptions> Allocator;
+		TUniquePtr<Ort::SessionOptions> SessionOptions;
+		TUniquePtr<Ort::MemoryInfo> AllocatorInfo;
+
+		/** IO ORT-related variables */
+		TArray<ONNXTensorElementDataType> InputTensorsORTType;
+		TArray<ONNXTensorElementDataType> OutputTensorsORTType;
+
+		TArray<const char*> InputTensorNames;
+		TArray<const char*> OutputTensorNames;
+
+		TArray<NNECore::Internal::FTensor> InputTensors;
+		TArray<NNECore::Internal::FTensor> OutputTensors;
+	
+		virtual bool InitializedAndConfigureMembers();
+		bool ConfigureTensors(const bool InIsInput);
+
+		/**
+		 * Statistics-related members used for GetLastRunTimeMSec(), GetRunStatistics(), GetInputMemoryTransferStats(), ResetStats().
+		 */
+		UE::NNEProfiling::Internal::FStatisticsEstimator RunStatisticsEstimator;
+		UE::NNEProfiling::Internal::FStatisticsEstimator InputTransferStatisticsEstimator;
+
+	};
+
+	class FModelORTCpu : public FModelORT
+	{
+	public:
+		FModelORTCpu() {}
+		FModelORTCpu(Ort::Env* InORTEnvironment, const FRuntimeConf& InRuntimeConf) : FModelORT(InORTEnvironment, InRuntimeConf) {}
+		virtual ~FModelORTCpu() = default;
+	private:
+		virtual bool InitializedAndConfigureMembers();
+	};
+
+#if PLATFORM_WINDOWS
+	class FModelORTDml : public FModelORT
+	{
+	public:
+		FModelORTDml() {}
+		FModelORTDml(Ort::Env* InORTEnvironment, const FRuntimeConf& InRuntimeConf) : FModelORT(InORTEnvironment, InRuntimeConf) {}
+		virtual ~FModelORTDml() = default;
+	private:
+		virtual bool InitializedAndConfigureMembers();
+	};
+#endif
+	
+} // UE::NNERuntimeORT::Private
\ No newline at end of file
diff --git a/Engine/Plugins/Experimental/NNE/Source/NNERuntimeORT/Private/NNERuntimeORTModule.cpp b/Engine/Plugins/Experimental/NNE/Source/NNERuntimeORT/Private/NNERuntimeORTModule.cpp
new file mode 100644
--- /dev/null
+++ b/Engine/Plugins/Experimental/NNE/Source/NNERuntimeORT/Private/NNERuntimeORTModule.cpp
@@ -0,0 +1,112 @@
+// Copyright Epic Games, Inc. All Rights Reserved.
+
+#include "NNERuntimeORTModule.h"
+#include "NNECore.h"
+#include "NNERuntimeORT.h"
+#include "HAL/FileManager.h"
+#include "HAL/PlatformProcess.h"
+#include "Misc/Paths.h"
+#include "UObject/WeakInterfacePtr.h"
+
+#include "NNXThirdPartyWarningDisabler.h"
+NNX_THIRD_PARTY_INCLUDES_START
+#undef check
+#undef TEXT
+#include "core/session/onnxruntime_cxx_api.h"
+NNX_THIRD_PARTY_INCLUDES_END
+
+void FNNERuntimeORTModule::StartupModule()
+{
+	TArray<FString> DLLFileNames
+	{
+		"onnxruntime",
+	};
+
+	//IPluginManager& PluginManager = IPluginManager::Get();
+	//const FString PluginDir = PluginManager.FindPlugin(TEXT("NNXIncubator"))->GetBaseDir();
+
+	//const FString ORTDefaultModuleDir = PluginDir / "ThirdParty//ORTDefault";
+	//const FString ORTDefaultRuntimeBinPath = ORTDefaultModuleDir / TEXT(PREPROCESSOR_TO_STRING(ORTDEFAULT_PLATFORM_PATH));;
+	const FString ORTDefaultRuntimeBinPath = TEXT(PREPROCESSOR_TO_STRING(ORTDEFAULT_PLATFORM_BIN_PATH));
+
+	FPlatformProcess::PushDllDirectory(*ORTDefaultRuntimeBinPath);
+
+	TArray<void*> DLLHandles;
+
+	for (FString DLLFileName : DLLFileNames)
+	{
+#if PLATFORM_WINDOWS
+		const FString DLLFilePath = ORTDefaultRuntimeBinPath / DLLFileName + ".dll";
+#elif PLATFORM_LINUX
+		const FString DLLFilePath = ORTDefaultRuntimeBinPath / "lib" + DLLFileName + ".so";
+#elif PLATFORM_MAC
+		const FString DLLFilePath = ORTDefaultRuntimeBinPath / "lib" + DLLFileName + ".dylib";
+#endif
+
+		// Sanity check
+		if (!FPaths::FileExists(DLLFilePath))
+		{
+			const FString ErrorMessage = FString::Format(TEXT("DLL file not found in \"{0}\"."),
+				{ IFileManager::Get().ConvertToAbsolutePathForExternalAppForRead(*DLLFilePath) });
+			//UE_LOG(LogNNX, Error, TEXT("ORT StartupModule(): %s"), *ErrorMessage);
+			//checkf(false, TEXT("%s"), *ErrorMessage);
+		}
+		DLLHandles.Add(FPlatformProcess::GetDllHandle(*DLLFilePath));
+
+	}
+	FPlatformProcess::PopDllDirectory(*ORTDefaultRuntimeBinPath);
+
+	Ort::InitApi();
+	
+	// NNE runtime ORT Cpu startup
+	NNERuntimeORTCpu = NewObject<UNNERuntimeORTCpuImpl>();
+	if (NNERuntimeORTCpu.IsValid())
+	{
+		TWeakInterfacePtr<INNERuntime> RuntimeCPUInterface(NNERuntimeORTCpu.Get());
+
+		NNERuntimeORTCpu->Init();
+		NNERuntimeORTCpu->AddToRoot();
+		UE::NNECore::RegisterRuntime(RuntimeCPUInterface);
+	}
+
+#if PLATFORM_WINDOWS
+	// NNE runtime ORT Dml startup
+	NNERuntimeORTDml = NewObject<UNNERuntimeORTDmlImpl>();
+	if (NNERuntimeORTDml.IsValid())
+	{
+		TWeakInterfacePtr<INNERuntime> RuntimeDmlInterface(NNERuntimeORTDml.Get());
+
+		NNERuntimeORTDml->Init();
+		NNERuntimeORTDml->AddToRoot();
+		UE::NNECore::RegisterRuntime(RuntimeDmlInterface);
+	}
+#endif
+}
+
+void FNNERuntimeORTModule::ShutdownModule()
+{
+	// NNE runtime ORT Cpu shutdown
+	if (NNERuntimeORTCpu.IsValid())
+	{
+		TWeakInterfacePtr<INNERuntime> RuntimeCPUInterface(NNERuntimeORTCpu.Get());
+
+		UE::NNECore::UnregisterRuntime(RuntimeCPUInterface);
+		NNERuntimeORTCpu->RemoveFromRoot();
+		NNERuntimeORTCpu = TWeakObjectPtr<UNNERuntimeORTCpuImpl>(nullptr);
+	}
+
+#if PLATFORM_WINDOWS
+	// NNE runtime ORT Dml shutdown
+	if (NNERuntimeORTDml.IsValid())
+	{
+		TWeakInterfacePtr<INNERuntime> RuntimeDmlInterface(NNERuntimeORTDml.Get());
+
+		UE::NNECore::UnregisterRuntime(RuntimeDmlInterface);
+		NNERuntimeORTDml->RemoveFromRoot();
+		NNERuntimeORTDml = TWeakObjectPtr<UNNERuntimeORTDmlImpl>(nullptr);
+	}
+#endif
+
+}
+
+IMPLEMENT_MODULE(FNNERuntimeORTModule, NNERuntimeORT);
\ No newline at end of file
diff --git a/Engine/Plugins/Experimental/NNE/Source/NNERuntimeORT/Private/NNERuntimeORTModule.h b/Engine/Plugins/Experimental/NNE/Source/NNERuntimeORT/Private/NNERuntimeORTModule.h
new file mode 100644
--- /dev/null
+++ b/Engine/Plugins/Experimental/NNE/Source/NNERuntimeORT/Private/NNERuntimeORTModule.h
@@ -0,0 +1,24 @@
+// Copyright Epic Games, Inc. All Rights Reserved.
+
+#pragma once
+
+#include "CoreMinimal.h"
+#include "Modules/ModuleManager.h"
+#include "Templates/SharedPointer.h"
+#include "UObject/WeakObjectPtr.h"
+#include "UObject/WeakObjectPtrTemplates.h"
+
+class UNNERuntimeORTDmlImpl;
+class UNNERuntimeORTCpuImpl;
+
+class FNNERuntimeORTModule : public IModuleInterface
+{
+
+public:
+	TWeakObjectPtr<UNNERuntimeORTCpuImpl> NNERuntimeORTCpu{ nullptr };
+	TWeakObjectPtr<UNNERuntimeORTDmlImpl> NNERuntimeORTDml{ nullptr };
+
+	// Begin IModuleInterface
+	virtual void StartupModule() override;
+	virtual void ShutdownModule() override;
+};
\ No newline at end of file
diff --git a/Engine/Plugins/Experimental/NNE/Source/NNERuntimeORT/Private/NNERuntimeORTProviders.h b/Engine/Plugins/Experimental/NNE/Source/NNERuntimeORT/Private/NNERuntimeORTProviders.h
new file mode 100644
--- /dev/null
+++ b/Engine/Plugins/Experimental/NNE/Source/NNERuntimeORT/Private/NNERuntimeORTProviders.h
@@ -0,0 +1,48 @@
+// Copyright Epic Games, Inc. All Rights Reserved.
+// Copyright (c) Microsoft Corporation. All rights reserved.
+// Licensed under the MIT License.
+
+#pragma once
+
+#include "NNXThirdPartyWarningDisabler.h"
+NNX_THIRD_PARTY_INCLUDES_START
+#undef check
+#undef TEXT
+
+#include "core/providers/cpu/cpu_provider_factory.h"
+
+#ifdef USE_DNNL
+#include "core/providers/dnnl/dnnl_provider_factory.h"
+#endif
+#ifdef USE_NUPHAR
+#include "core/providers/nuphar/nuphar_provider_factory.h"
+#endif
+#ifdef USE_TVM
+#include "core/providers/tvm/tvm_provider_factory.h"
+#endif
+#ifdef USE_TENSORRT
+#include "core/providers/tensorrt/tensorrt_provider_factory.h"
+#endif
+#ifdef USE_OPENVINO
+#include "core/providers/openvino/openvino_provider_factory.h"
+#endif
+#ifdef USE_NNAPI
+#include "core/providers/nnapi/nnapi_provider_factory.h"
+#endif
+#ifdef USE_COREML
+#include "core/providers/coreml/coreml_provider_factory.h"
+#endif
+#ifdef USE_DML
+#include "core/providers/dml/dml_provider_factory.h"
+#endif
+#ifdef USE_ACL
+#include "core/providers/acl/acl_provider_factory.h"
+#endif
+#ifdef USE_ARMNN
+#include "core/providers/armnn/armnn_provider_factory.h"
+#endif
+#ifdef USE_MIGRAPHX
+#include "core/providers/migraphx/migraphx_provider_factory.h"
+#endif
+
+NNX_THIRD_PARTY_INCLUDES_END
\ No newline at end of file
diff --git a/Engine/Plugins/Experimental/NNE/Source/NNERuntimeORT/Private/NNERuntimeORTUtils.h b/Engine/Plugins/Experimental/NNE/Source/NNERuntimeORT/Private/NNERuntimeORTUtils.h
new file mode 100644
--- /dev/null
+++ b/Engine/Plugins/Experimental/NNE/Source/NNERuntimeORT/Private/NNERuntimeORTUtils.h
@@ -0,0 +1,214 @@
+// Copyright Epic Games, Inc. All Rights Reserved.
+
+#pragma once
+
+#include "HAL/UnrealMemory.h"
+#include "NNECore.h"
+#include "NNECoreRuntimeCPU.h"
+#include "NNECoreTypes.h"
+#include "NNECoreTensor.h"
+#include "NNXThirdPartyWarningDisabler.h"
+NNX_THIRD_PARTY_INCLUDES_START
+#undef check
+#undef TEXT
+#include "core/session/onnxruntime_cxx_api.h"
+NNX_THIRD_PARTY_INCLUDES_END
+
+DECLARE_STATS_GROUP(TEXT("NNE"), STATGROUP_NNE, STATCAT_Advanced);
+
+namespace UE::NNERuntimeORT::Private
+{
+
+	using TypeInfoORT = std::pair<ENNETensorDataType, uint64>;
+
+	inline TypeInfoORT TranslateTensorTypeORTToNNE(unsigned int OrtDataType)
+	{
+		ENNETensorDataType DataType = ENNETensorDataType::None;
+		uint64 ElementSize = 0;
+
+		switch (OrtDataType) {
+		case ONNX_TENSOR_ELEMENT_DATA_TYPE_UNDEFINED:
+			DataType = ENNETensorDataType::None;
+			ElementSize = 0;
+			break;
+
+		case ONNX_TENSOR_ELEMENT_DATA_TYPE_FLOAT:
+			DataType = ENNETensorDataType::Float;
+			ElementSize = sizeof(float);
+			break;
+
+		case ONNX_TENSOR_ELEMENT_DATA_TYPE_UINT8:
+			DataType = ENNETensorDataType::UInt8;
+			ElementSize = sizeof(uint8);
+			break;
+		case ONNX_TENSOR_ELEMENT_DATA_TYPE_INT8:
+			DataType = ENNETensorDataType::Int8;
+			ElementSize = sizeof(int8);
+			break;
+		case ONNX_TENSOR_ELEMENT_DATA_TYPE_UINT16:
+			DataType = ENNETensorDataType::UInt16;
+			ElementSize = sizeof(uint16);
+			break;
+		case ONNX_TENSOR_ELEMENT_DATA_TYPE_INT16:
+			DataType = ENNETensorDataType::Int16;
+			ElementSize = sizeof(int16);
+			break;
+
+		case ONNX_TENSOR_ELEMENT_DATA_TYPE_INT32:
+			DataType = ENNETensorDataType::Int32;
+			ElementSize = sizeof(int32);
+			break;
+
+		case ONNX_TENSOR_ELEMENT_DATA_TYPE_INT64:
+			DataType = ENNETensorDataType::Int64;
+			ElementSize = sizeof(int64);
+			break;
+
+		case ONNX_TENSOR_ELEMENT_DATA_TYPE_STRING:
+			DataType = ENNETensorDataType::Char;
+			ElementSize = sizeof(char);
+			break;
+
+		case ONNX_TENSOR_ELEMENT_DATA_TYPE_BOOL:
+			DataType = ENNETensorDataType::Boolean;
+			ElementSize = sizeof(bool);
+			break;
+
+		case ONNX_TENSOR_ELEMENT_DATA_TYPE_FLOAT16:
+			DataType = ENNETensorDataType::Half;
+			ElementSize = 2;
+			break;
+
+		case ONNX_TENSOR_ELEMENT_DATA_TYPE_DOUBLE:
+			DataType = ENNETensorDataType::Double;
+			ElementSize = sizeof(double);
+			break;
+
+		case ONNX_TENSOR_ELEMENT_DATA_TYPE_UINT32:
+			DataType = ENNETensorDataType::UInt32;
+			ElementSize = sizeof(uint32);
+			break;
+
+		case ONNX_TENSOR_ELEMENT_DATA_TYPE_UINT64:
+			DataType = ENNETensorDataType::UInt64;
+			ElementSize = sizeof(uint64);
+			break;
+
+		case ONNX_TENSOR_ELEMENT_DATA_TYPE_COMPLEX64:
+			DataType = ENNETensorDataType::Complex64;
+			ElementSize = 8;
+			break;
+
+		case ONNX_TENSOR_ELEMENT_DATA_TYPE_COMPLEX128:
+			DataType = ENNETensorDataType::Complex128;
+			ElementSize = 16;
+			break;
+
+		case ONNX_TENSOR_ELEMENT_DATA_TYPE_BFLOAT16:
+			DataType = ENNETensorDataType::BFloat16;
+			ElementSize = 2;
+			break;
+
+		default:
+			DataType = ENNETensorDataType::None;
+			break;
+		}
+
+		return TypeInfoORT{ DataType, ElementSize };
+	}
+
+	inline void BindTensorsToORT(
+		TConstArrayView<NNECore::FTensorBindingCPU> InBindingTensors,
+		TConstArrayView<NNECore::Internal::FTensor> InTensors,
+		TConstArrayView<ONNXTensorElementDataType> InTensorsORTType,
+		const Ort::MemoryInfo* InAllocatorInfo,
+		TArray<Ort::Value>& OutOrtTensors
+	)
+	{
+		const uint32 NumBinding = InBindingTensors.Num();
+		const uint32 NumDescriptors = InTensors.Num();
+
+		if (NumBinding != NumDescriptors)
+		{
+			UE_LOG(LogNNE, Error, TEXT("BindTensorsToORT: Number of Bindings is different of Descriptors."));
+			return;
+		}
+
+		for (uint32 Index = 0; Index < NumBinding; ++Index)
+		{
+			const NNECore::FTensorBindingCPU& Binding = InBindingTensors[Index];
+			const NNECore::Internal::FTensor& Tensor = InTensors[Index];
+			const ONNXTensorElementDataType CurrentORTType = InTensorsORTType[Index];
+
+			TUniquePtr<int64_t[]> SizesInt64t;
+			SizesInt64t = MakeUnique<int64_t[]>(Tensor.GetShape().Rank());
+			for (int32 DimIndex = 0; DimIndex < Tensor.GetShape().Rank(); ++DimIndex)
+			{
+				SizesInt64t.Get()[DimIndex] = Tensor.GetShape().GetData()[DimIndex];
+			}
+
+			const uint64 ByteCount{ InTensors[Index].GetDataSize() };
+			const uint32 ArrayDimensions{ (uint32)Tensor.GetShape().Rank() };
+			OutOrtTensors.Emplace(
+				Ort::Value::CreateTensor(
+					*InAllocatorInfo,
+					Binding.Data,
+					ByteCount,
+					SizesInt64t.Get(),
+					ArrayDimensions,
+					CurrentORTType
+				)
+			);
+		}
+	}
+
+	inline void CopyFromORTToBindings(
+		TConstArrayView<Ort::Value> InOrtTensors,
+		TConstArrayView<NNECore::FTensorBindingCPU> InBindingTensors,
+		TConstArrayView<NNECore::FTensorDesc> InTensorDescs,
+		TArray<NNECore::Internal::FTensor>& OutTensors
+	)
+	{
+		const uint32 NumBinding = InOrtTensors.Num();
+		const uint32 NumDescriptors = InTensorDescs.Num();
+
+		if (NumBinding != NumDescriptors)
+		{
+			UE_LOG(LogNNE, Error, TEXT("CopyFromORTToBindings: Number of Bindings is different of Descriptors."));
+			return;
+		}
+
+		for (uint32 Index = 0; Index < NumBinding; ++Index)
+		{
+			const NNECore::FTensorBindingCPU& Binding = InBindingTensors[Index];
+			const NNECore::FTensorDesc& TensorDesc = InTensorDescs[Index];
+			const Ort::Value& OrtTensor = InOrtTensors[Index];
+			const std::vector<int64_t>& OrtShape = OrtTensor.GetTensorTypeAndShapeInfo().GetShape();
+
+			TArray<uint32> ShapeData;
+			for (int32 DimIndex = 0; DimIndex < OrtShape.size(); ++DimIndex)
+			{
+				check(OrtShape[DimIndex] >= 0);
+				ShapeData.Add(OrtShape[DimIndex]);
+			}
+
+			NNECore::FTensorShape Shape = NNECore::FTensorShape::Make(ShapeData);
+			NNECore::Internal::FTensor Tensor = NNECore::Internal::FTensor::Make(TensorDesc.GetName(), Shape, TensorDesc.GetDataType());
+			OutTensors.Emplace(Tensor);
+
+			const void* OrtTensorData = OrtTensor.GetTensorData<void>();
+			const uint64 DataSize = FMath::Min(Tensor.GetDataSize(), Binding.SizeInBytes);
+
+			if (DataSize > 0)
+			{
+				FMemory::Memcpy(Binding.Data, OrtTensorData, DataSize);
+			}
+
+			if (Tensor.GetDataSize() > Binding.SizeInBytes)
+			{
+				UE_LOG(LogNNE, Error, TEXT("CopyFromORTToBindings: Binding buffer was not large enough to contain all of the data, only first %d bytes were copied."), DataSize);
+			}
+		}
+	}
+
+} // UE::NNERuntimeORT::Private
\ No newline at end of file
diff --git a/Engine/Plugins/Experimental/NNE/Source/NNEUtils/Private/NNEUtilsModelBuilderNNE.cpp b/Engine/Plugins/Experimental/NNE/Source/NNEUtils/Private/NNEUtilsModelBuilderNNE.cpp
--- a/Engine/Plugins/Experimental/NNE/Source/NNEUtils/Private/NNEUtilsModelBuilderNNE.cpp
+++ b/Engine/Plugins/Experimental/NNE/Source/NNEUtils/Private/NNEUtilsModelBuilderNNE.cpp
@@ -98,9 +98,8 @@ public:
 	virtual bool End(TArray<uint8>& Data) override
 	{
 		// This is for debugging purposes
-		FModelPrinterNNE Printer;
-
-		Printer.Visit(Format);
+		//FModelPrinterNNE Printer;
+		//Printer.Visit(Format);
 
 		FMemoryWriter Writer(Data);
 
diff --git a/Engine/Plugins/Experimental/NNE/Source/NNXCore/Private/NNECoreModelData.cpp b/Engine/Plugins/Experimental/NNE/Source/NNXCore/Private/NNECoreModelData.cpp
--- a/Engine/Plugins/Experimental/NNE/Source/NNXCore/Private/NNECoreModelData.cpp
+++ b/Engine/Plugins/Experimental/NNE/Source/NNXCore/Private/NNECoreModelData.cpp
@@ -102,7 +102,6 @@ TConstArrayView<uint8> UNNEModelData::GetModelData(const FString& RuntimeName)
 	TArray<uint8>* LocalData = ModelData.Find(RuntimeName);
 	if (LocalData)
 	{
-		UE_LOG(LogNNX, Display, TEXT("UNNEModelData: Using locally cached model data"));
 		return TConstArrayView<uint8>(LocalData->GetData(), LocalData->Num());
 	}
 
@@ -111,9 +110,7 @@ TConstArrayView<uint8> UNNEModelData::GetModelData(const FString& RuntimeName)
 	if (RemoteData.GetSize() > 0)
 	{
 		ModelData.Add(RuntimeName, TArray<uint8>((uint8*)RemoteData.GetData(), RemoteData.GetSize()));
-
-		UE_LOG(LogNNX, Display, TEXT("UNNEModelData: Using remote cached model data"));
-
+		
 		TArray<uint8>* CachedRemoteData = ModelData.Find(RuntimeName);
 		return TConstArrayView<uint8>(CachedRemoteData->GetData(), CachedRemoteData->Num());
 	}
@@ -126,7 +123,6 @@ TConstArrayView<uint8> UNNEModelData::GetModelData(const FString& RuntimeName)
 	}
 
 	// Cache the model and put it into ddc
-	UE_LOG(LogNNX, Display, TEXT("UNNEModelData: Caching created model data"));
 	ModelData.Add(RuntimeName, CreatedData);
 
 	FSharedBuffer SharedBuffer = MakeSharedBufferFromArray(MoveTemp(CreatedData));
@@ -140,7 +136,6 @@ TConstArrayView<uint8> UNNEModelData::GetModelData(const FString& RuntimeName)
 	TArray<uint8>* Data = ModelData.Find(RuntimeName);
 	if (Data)
 	{
-		UE_LOG(LogTemp, Display, TEXT("UNNEModelData: Using cooked model data"));
 		return TConstArrayView<uint8>(Data->GetData(), Data->Num());
 	}
 	return {};
diff --git a/Engine/Plugins/Experimental/NNE/Source/NNXQA/Private/NNXQAUtils.cpp b/Engine/Plugins/Experimental/NNE/Source/NNXQA/Private/NNXQAUtils.cpp
--- a/Engine/Plugins/Experimental/NNE/Source/NNXQA/Private/NNXQAUtils.cpp
+++ b/Engine/Plugins/Experimental/NNE/Source/NNXQA/Private/NNXQAUtils.cpp
@@ -640,6 +640,40 @@ namespace Test
 		return bTestSuceeded;
 	}
 
+	bool RunTestInferenceAndCompareToRef(const FTests::FTestSetup& TestSetup, const FString& RuntimeName, const FNNIModelRaw& ONNXModel,
+		TArrayView<TArray<char>> RefOutputMemBuffers, TArrayView<UE::NNECore::Internal::FTensor> RefOutputTensors)
+	{
+		TArray<TArray<char>> OutputMemBuffers;
+		TArray<UE::NNECore::Internal::FTensor> OutputTensors;
+
+		const float AbsoluteTolerance = TestSetup.GetAbsoluteToleranceForRuntime(RuntimeName);
+		const float RelativeTolerance = TestSetup.GetRelativeToleranceForRuntime(RuntimeName);
+		bool bTestSuceeded = true;
+
+		if (!RunTestInferenceCPU(ONNXModel, TestSetup, RuntimeName, OutputTensors, OutputMemBuffers))
+		{
+			UE_LOG(LogNNX, Error, TEXT("Error running inference for engine %s."), *RuntimeName);
+			return false;
+		}
+
+		if (OutputTensors.Num() == RefOutputTensors.Num())
+		{
+			for (int i = 0; i < OutputTensors.Num(); ++i)
+			{
+				bTestSuceeded &= VerifyTensorResult(
+					RefOutputTensors[i], RefOutputMemBuffers[i],
+					OutputTensors[i], OutputMemBuffers[i],
+					AbsoluteTolerance, RelativeTolerance);
+			}
+		}
+		else
+		{
+			UE_LOG(LogNNX, Error, TEXT("Expecting %d output tensor(s), got %d."), RefOutputTensors.Num(), OutputTensors.Num());
+			bTestSuceeded = false;
+		}
+		return bTestSuceeded;
+	}
+
 	bool CompareONNXModelInferenceAcrossRuntimes(const FNNIModelRaw& ONNXModel, const FNNIModelRaw& ONNXModelVariadic, const FTests::FTestSetup& TestSetup, const FString& RuntimeFilter)
 	{
 		FString CurrentPlatform = UGameplayStatics::GetPlatformName();
@@ -677,34 +711,68 @@ namespace Test
 		}
 		if (!bAllTestsSucceeded)
 		{
-			UE_LOG(LogNNX, Error, TEXT("Expecting output from test setup are no match by reference engine."));
+			UE_LOG(LogNNX, Error, TEXT("Expecting output from test setup are no matched by reference runtime."));
 		}
 
 		// Test against other runtime
-		TArray<IRuntime*> Runtimes;
-
-		Runtimes = GetAllRuntimes();
-		for (auto Runtime : Runtimes)
+		//NNX Registry
+		for (auto Runtime : GetAllRuntimes())
 		{
 			const FString& RuntimeName = Runtime->GetRuntimeName();
-			if ((RuntimeName == RefName) || RuntimeName == TEXT("NNXRuntimeCPU"))
+			if (RuntimeName == TEXT("NNXRuntimeCPU"))
 			{
 				continue;
 			}
-
 			if (!RuntimeFilter.IsEmpty() && !RuntimeFilter.Contains(RuntimeName))
 			{
 				continue;
 			}
+			FString TestResult;
+			if (TestSetup.AutomationExcludedRuntime.Contains(RuntimeName) ||
+				TestSetup.AutomationExcludedPlatformRuntimeCombination.Contains(TPair<FString, FString>(CurrentPlatform, RuntimeName)))
+			{
+				TestResult = TEXT("skipped (by config)");
+			}
+			else
+			{
+				bool bShouldRunVariadicTest = (ONNXModelVariadic.Format != ENNXInferenceFormat::Invalid);
+				bShouldRunVariadicTest &= !(RuntimeName == "NNXRuntimeDML");
 
-			if (RuntimeName == "NNXRuntimeORTCuda")
+				bool bTestSuceeded = RunTestInferenceAndCompareToRef(TestSetup, Runtime, ONNXModel, RefOutputMemBuffers, RefOutputTensors);
+
+				if (bShouldRunVariadicTest)
+				{
+					if (!bTestSuceeded)
+					{
+						UE_LOG(LogNNX, Error, TEXT("Failed running static test."));
+					}
+					bool bVariadicTestSuceeded = RunTestInferenceAndCompareToRef(TestSetup, Runtime, ONNXModelVariadic, RefOutputMemBuffers, RefOutputTensors);
+					if (!bVariadicTestSuceeded)
+					{
+						bTestSuceeded = false;
+						UE_LOG(LogNNX, Error, TEXT("Failed running variadic test."));
+					}
+				}
+				TestResult = bTestSuceeded ? TEXT("SUCCESS") : TEXT("FAILED");
+				bAllTestsSucceeded &= bTestSuceeded;
+			}
+
+			UE_LOG(LogNNX, Display, TEXT("  %s tests: %s"), *RuntimeName, *TestResult);
+		}
+
+		//NNE Registry
+		for (auto Runtime : UE::NNECore::GetAllRuntimes())
+		{
+			const FString& RuntimeName = Runtime->GetRuntimeName();
+			if (RuntimeName == RefName)
+			{
+				continue;
+			}
+			if (!RuntimeFilter.IsEmpty() && !RuntimeFilter.Contains(RuntimeName))
 			{
-				//TODO Reactivate tests for NNXRuntimeORTCuda runtime. Skipped for now as we wait for legal approval for the dlls.
 				continue;
 			}
-
 			FString TestResult;
-
 			if (TestSetup.AutomationExcludedRuntime.Contains(RuntimeName) ||
 				TestSetup.AutomationExcludedPlatformRuntimeCombination.Contains(TPair<FString, FString>(CurrentPlatform, RuntimeName)))
 			{
@@ -713,9 +781,9 @@ namespace Test
 			else
 			{
 				bool bShouldRunVariadicTest = (ONNXModelVariadic.Format != ENNXInferenceFormat::Invalid);
-				bShouldRunVariadicTest &= !(RuntimeName == "NNXRuntimeDML");
+				bShouldRunVariadicTest &= !(RuntimeName == "NNEXRuntimeDML");
 
-				bool bTestSuceeded = RunTestInferenceAndCompareToRef(TestSetup, Runtime, ONNXModel, RefOutputMemBuffers, RefOutputTensors);
+				bool bTestSuceeded = RunTestInferenceAndCompareToRef(TestSetup, Runtime->GetRuntimeName(), ONNXModel, RefOutputMemBuffers, RefOutputTensors);
 
 				if (bShouldRunVariadicTest)
 				{
@@ -723,7 +791,7 @@ namespace Test
 					{
 						UE_LOG(LogNNX, Error, TEXT("Failed running static test."));
 					}
-					bool bVariadicTestSuceeded = RunTestInferenceAndCompareToRef(TestSetup, Runtime, ONNXModelVariadic, RefOutputMemBuffers, RefOutputTensors);
+					bool bVariadicTestSuceeded = RunTestInferenceAndCompareToRef(TestSetup, Runtime->GetRuntimeName(), ONNXModelVariadic, RefOutputMemBuffers, RefOutputTensors);
 					if (!bVariadicTestSuceeded)
 					{
 						bTestSuceeded = false;
@@ -736,6 +804,7 @@ namespace Test
 
 			UE_LOG(LogNNX, Display, TEXT("  %s tests: %s"), *RuntimeName, *TestResult);
 		}
+		
 		return bAllTestsSucceeded;
 	}
 	FTests::FTestSetup& FTests::AddTest(const FString & Category, const FString & ModelOrOperatorName, const FString & TestSuffix)
diff --git a/Engine/Plugins/Experimental/NNE/Source/NNXQA/Resources/NNXQATestDesc.json b/Engine/Plugins/Experimental/NNE/Source/NNXQA/Resources/NNXQATestDesc.json
--- a/Engine/Plugins/Experimental/NNE/Source/NNXQA/Resources/NNXQATestDesc.json
+++ b/Engine/Plugins/Experimental/NNE/Source/NNXQA/Resources/NNXQATestDesc.json
@@ -9,7 +9,7 @@
 					"target": "resnet50-v2-7",
 					"tags": [ "SmokeTest" ],
 					"runtimes": [
-						{ "name": "NNXRuntimeORTDml" },
+						{ "name": "NNERuntimeORTDmlEditor" },
 						{ "name": "NNXRuntimeHlsl",  "skip": true, "comment": "Missing BatchNormalisation operator." },
 						{ "name": "NNXRuntimeDml",  "skip": true, "comment": "Model not supported yet." }
 					],
@@ -19,8 +19,8 @@
 					"target": "Mosaic-9",
 					"tags": [ "SmokeTest" ],
 					"runtimes": [
-						{ "name": "NNXRuntimeORTDml" },
-						{ "name": "NNXRuntimeORTCpu" },
+						{ "name": "NNERuntimeORTDmlEditor" },
+						{ "name": "NNERuntimeORTCpuEditor" },
 						{ "name": "NNXRuntimeHlsl" },
 						{ "name": "NNXRuntimeDml",  "skip": true, "comment": "Model not supported yet." }
 					],
@@ -36,7 +36,7 @@
 					"target": "NeuralMorphModel_global",
 					"tags": [ "SmokeTest" ],
 					"runtimes": [
-						{ "name": "NNXRuntimeORTDml" },
+						{ "name": "NNERuntimeORTDmlEditor" },
 						{ "name": "NNXRuntimeHlsl" },
 						{ "name": "NNXRuntimeDml" }
 					],
@@ -46,7 +46,7 @@
 					"target": "NeuralMorphModel_global_Dax",
 					"tags": [ "SmokeTest" ],
 					"runtimes": [
-						{ "name": "NNXRuntimeORTDml", "skip": true, "comment": "" },
+						{ "name": "NNERuntimeORTDmlEditor", "skip": true, "comment": "" },
 						{ "name": "NNXRuntimeHlsl",	"skip": true, "comment": "Model not supported yet."	},
 						{ "name": "NNXRuntimeDml",  "skip": true, "comment": "Model not supported yet." }
 					]
@@ -55,7 +55,7 @@
 					"target": "NeuralMorphModel_Local",
 					"tags": [ "SmokeTest" ],
 					"runtimes": [
-						{ "name": "NNXRuntimeORTDml", "absolute_tolerance": 1e-6 },
+						{ "name": "NNERuntimeORTDmlEditor", "absolute_tolerance": 1e-6 },
 						{ "name": "NNXRuntimeHlsl",  "skip": true, "comment": "Missing Reshape operator." },
 						{ "name": "NNXRuntimeDml",  "skip": true, "comment": "Model not supported yet." }
 					]
@@ -63,7 +63,7 @@
 				{
 					"target": "VertexDeltaModel",
 					"runtimes": [
-						{ "name": "NNXRuntimeORTDml", "absolute_tolerance": 1e-4 },
+						{ "name": "NNERuntimeORTDmlEditor", "absolute_tolerance": 1e-4 },
 						{ "name": "NNXRuntimeHlsl", "absolute_tolerance": 1e-4 },
 						{ "name": "NNXRuntimeDml",  "skip": true, "comment": "Model not supported yet." }
 					]
@@ -125,23 +125,23 @@
 			"targets": [
 				{	"target": "Abs", "additional_datasets": [ "Elementwise.Slow.OneInput" ]},
 				{	"target": "Acos", "tags": [ "SmokeTest" ], "runtimes": [
-						{ "name": "NNXRuntimeORTDml", "absolute_tolerance":  1e-4},
+						{ "name": "NNERuntimeORTDmlEditor", "absolute_tolerance":  1e-4},
 						{ "name": "NNXRuntimeDml",    "absolute_tolerance":  1e-4},
 						{ "name": "NNXRuntimeHlsl",   "absolute_tolerance":  2e-4}
 					]},
 				{	"target": "Acosh" },
 				{	"target": "Asin", "runtimes": [
-						{ "name": "NNXRuntimeORTDml", "absolute_tolerance":  1e-4},
+						{ "name": "NNERuntimeORTDmlEditor", "absolute_tolerance":  1e-4},
 						{ "name": "NNXRuntimeDml",    "absolute_tolerance":  1e-4},
 						{ "name": "NNXRuntimeHlsl",   "absolute_tolerance":  2e-4}
 					]},
 				{	"target": "Asinh", "runtimes": [
-						{ "name": "NNXRuntimeORTDml", "absolute_tolerance":  1e-4},
+						{ "name": "NNERuntimeORTDmlEditor", "absolute_tolerance":  1e-4},
 						{ "name": "NNXRuntimeDml",    "absolute_tolerance":  1e-4},
 						{ "name": "NNXRuntimeHlsl",   "absolute_tolerance":  2e-4}
 					]},
 				{	"target": "Atan", "runtimes": [
-						{ "name": "NNXRuntimeORTDml", "absolute_tolerance":  1e-4},
+						{ "name": "NNERuntimeORTDmlEditor", "absolute_tolerance":  1e-4},
 						{ "name": "NNXRuntimeDml",    "absolute_tolerance":  1e-4},
 						{ "name": "NNXRuntimeHlsl",   "absolute_tolerance":  1.5e-3}
 					]},
@@ -156,12 +156,12 @@
 				{	"target": "Cos" },
 				{	"target": "Cosh" },
 				{	"target": "Elu", "comment": "TODO test with not-default attribute", "runtimes": [
-						{ "name": "NNXRuntimeORTDml", "absolute_tolerance":  1e-4},
+						{ "name": "NNERuntimeORTDmlEditor", "absolute_tolerance":  1e-4},
 						{ "name": "NNXRuntimeDml",    "absolute_tolerance":  1e-4},
 						{ "name": "NNXRuntimeHlsl",   "absolute_tolerance":  1e-4}
 					]},
 				{	"target": "Erf", "tags": [ "SmokeTest" ], "runtimes": [
-						{ "name": "NNXRuntimeORTDml", "absolute_tolerance":  1e-4},
+						{ "name": "NNERuntimeORTDmlEditor", "absolute_tolerance":  1e-4},
 						{ "name": "NNXRuntimeDml",    "absolute_tolerance":  1e-4},
 						{ "name": "NNXRuntimeHlsl",   "absolute_tolerance":  5e-4, "comment": "NNXRuntimeHLSL Erf implementation is less precise than the DML one"}
 					]},
@@ -173,7 +173,7 @@
 				{	"target": "HardSwish", "runtimes": [ { "name": "NNXRuntimeDml", "skip": true, "comment": "Not implemented atm on NNXRuntimeDML."} ]},
 				{	"target": "LeakyRelu", "additional_datasets": [ "Elementwise.LeakyRelu.AttributeTest" ] },
 				{	"target": "Log", "runtimes": [
-						{ "name": "NNXRuntimeORTDml"},
+						{ "name": "NNERuntimeORTDmlEditor"},
 						{ "name": "NNXRuntimeDml"},
 						{ "name": "NNXRuntimeHlsl"}
 					], "absolute_tolerance": 1e-6 },
@@ -186,14 +186,14 @@
 				{	"target": "Relu", "tags": [ "SmokeTest" ] },
 				{	"target": "Round", "tags": [ "SmokeTest" ] },
 				{	"target": "Selu", "comment": "TODO test with not-default attribute", "runtimes": [
-						{ "name": "NNXRuntimeORTDml", "absolute_tolerance":  1e-4},
+						{ "name": "NNERuntimeORTDmlEditor", "absolute_tolerance":  1e-4},
 						{ "name": "NNXRuntimeDml",    "absolute_tolerance":  1e-4},
 						{ "name": "NNXRuntimeHlsl",   "absolute_tolerance":  1e-4}
 					]},
 				{	"target": "Sigmoid", "tags": [ "SmokeTest" ] },
 				{	"target": "Sign" },
 				{	"target": "Sin", "runtimes": [
-						{ "name": "NNXRuntimeORTDml", "absolute_tolerance":  1e-4},
+						{ "name": "NNERuntimeORTDmlEditor", "absolute_tolerance":  1e-4},
 						{ "name": "NNXRuntimeDml",    "absolute_tolerance":  1e-4},
 						{ "name": "NNXRuntimeHlsl",   "absolute_tolerance":  1e-4}
 					]},
@@ -202,12 +202,12 @@
 				{	"target": "Softsign", "tags": [ "SmokeTest" ] },
 				{	"target": "Sqrt", "tags": [ "SmokeTest" ] },
 				{	"target": "Tan", "runtimes": [
-						{ "name": "NNXRuntimeORTDml", "absolute_tolerance":  1e-4},
+						{ "name": "NNERuntimeORTDmlEditor", "absolute_tolerance":  1e-4},
 						{ "name": "NNXRuntimeDml",    "absolute_tolerance":  1e-4},
 						{ "name": "NNXRuntimeHlsl",   "absolute_tolerance":  1e-4}
 					]},
 				{	"target": "Tanh", "runtimes": [
-						{ "name": "NNXRuntimeORTDml", "absolute_tolerance":  1e-4},
+						{ "name": "NNERuntimeORTDmlEditor", "absolute_tolerance":  1e-4},
 						{ "name": "NNXRuntimeDml",    "absolute_tolerance":  1e-4},
 						{ "name": "NNXRuntimeHlsl",   "absolute_tolerance":  1e-4}
 					]}
@@ -281,7 +281,7 @@
 			"category": "Convolution",
 			"runtimes": [
 				{ "name": "NNXRuntimeDml", "skip": true	},
-				{ "name": "NNXRuntimeORTDml" },
+				{ "name": "NNERuntimeORTDmlEditor" },
 				{ "name": "NNXRuntimeHlsl" }
 			],
 			"targets": [
@@ -292,7 +292,7 @@
 			"category": "ConvTranspose",
 			"runtimes": [
 				{ "name": "NNXRuntimeDml", "skip": true	},
-				{ "name": "NNXRuntimeORTDml" },
+				{ "name": "NNERuntimeORTDmlEditor" },
 				{ "name": "NNXRuntimeHlsl" }
 			],
 			"targets": [
@@ -478,7 +478,7 @@
 				{ "inputs": [ { "shape": [ 3, 4 ] }, { "shape": [ 4, 5 ] } ], "outputs": [ { "shape": [ 3, 5 ] } ]}
 			],
 			"runtimes": [
-				{ "name": "NNXRuntimeORTDml", "skip": true, "comment": "Variable dimension without bias is failing for NNXRuntimeORTDml" }
+				{ "name": "NNERuntimeORTDmlEditor", "skip": true, "comment": "Variable dimension without bias is failing for NNERuntimeORTDmlEditor" }
 			]
 		},
 		{
@@ -488,7 +488,7 @@
 				{ "inputs": [ { "shape": [ 4, 3 ] }, { "shape": [ 4, 5 ] } ], "outputs": [ { "shape": [ 3, 5 ] } ]}
 			],
 			"runtimes": [
-				{ "name": "NNXRuntimeORTDml", "skip": true, "comment": "Variable dimension without bias is failing for NNXRuntimeORTDml" }
+				{ "name": "NNERuntimeORTDmlEditor", "skip": true, "comment": "Variable dimension without bias is failing for NNERuntimeORTDmlEditor" }
 			]
 		},
 		{
@@ -498,7 +498,7 @@
 				{ "inputs": [ { "shape": [ 3, 4 ] }, { "shape": [ 5, 4 ] } ], "outputs": [ { "shape": [ 3, 5 ] } ]}
 			],
 			"runtimes": [
-				{ "name": "NNXRuntimeORTDml", "skip": true, "comment": "Variable dimension without bias is failing for NNXRuntimeORTDml" }
+				{ "name": "NNERuntimeORTDmlEditor", "skip": true, "comment": "Variable dimension without bias is failing for NNERuntimeORTDmlEditor" }
 			]
 		},
 		{
@@ -508,7 +508,7 @@
 				{ "inputs": [ { "shape": [ 4, 3 ] }, { "shape": [ 5, 4 ] } ], "outputs": [ { "shape": [ 3, 5 ] } ]}
 			],
 			"runtimes": [
-				{ "name": "NNXRuntimeORTDml", "skip": true, "comment": "Variable dimension without bias is failing for NNXRuntimeORTDml" }
+				{ "name": "NNERuntimeORTDmlEditor", "skip": true, "comment": "Variable dimension without bias is failing for NNERuntimeORTDmlEditor" }
 			]
 		},
 		{
@@ -564,7 +564,7 @@
 			"name": "ConvTranspose.Autopad",
 			"datasets": [ { "inputs": [ { "shape": [ 1, 1, 3, 3 ] }, { "shape": [ 1, 2, 3, 3 ] } ], "outputs": [ { "shape": [ 1, 2, 6, 6 ] } ] } ],
 			"runtimes": [
-				{ "name": "NNXRuntimeORTDml", "skip": true, "comment": "TODO INVESTIGATE Returns 'The parameter is incorrect'" },
+				{ "name": "NNERuntimeORTDmlEditor", "skip": true, "comment": "TODO INVESTIGATE Returns 'The parameter is incorrect'" },
 				{ "name": "NNXRuntimeHlsl", "skip": true, "comment": "Fails until ORT 1.13, because they implemented 'auto_pad' differently (they flipped 'SAME_LOWER' and 'SAME_UPPER')" }
 			]
 		},
@@ -586,7 +586,7 @@
 		{
 			"name": "MatMul.Rank5",
 			"runtimes": [	
-				{ "name": "NNXRuntimeORTDml", "skip": true, "comment": "DML metacommand only support 2 to 4 for inputs rank" },
+				{ "name": "NNERuntimeORTDmlEditor", "skip": true, "comment": "DML metacommand only support 2 to 4 for inputs rank" },
 				{ "name": "NNXRuntimeHlsl"}
 			],
 			"datasets": [
@@ -655,9 +655,9 @@
 		{
 			"name": "InstanceNormalization",
 			"datasets": [
-				{ "inputs": [ { "shape": [ 1, 1, 3 ] }, { "shape": [ 1 ] }, { "shape": [ 1 ] } ], "outputs": [ { "shape": [ 1, 1, 3 ] } ], "runtimes": [ { "name": "NNXRuntimeORTDml", "skip": true, "comment": "Not supported" } ] },
-				{ "inputs": [ { "shape": [ 1, 2, 3 ] }, { "shape": [ 2 ] }, { "shape": [ 2 ] } ], "outputs": [ { "shape": [ 1, 2, 3 ] } ], "runtimes": [ { "name": "NNXRuntimeORTDml", "skip": true, "comment": "Not supported" } ] },
-				{ "inputs": [ { "shape": [ 2, 2, 3 ] }, { "shape": [ 2 ] }, { "shape": [ 2 ] } ], "outputs": [ { "shape": [ 2, 2, 3 ] } ], "runtimes": [ { "name": "NNXRuntimeORTDml", "skip": true, "comment": "Not supported" } ]},
+				{ "inputs": [ { "shape": [ 1, 1, 3 ] }, { "shape": [ 1 ] }, { "shape": [ 1 ] } ], "outputs": [ { "shape": [ 1, 1, 3 ] } ], "runtimes": [ { "name": "NNERuntimeORTDmlEditor", "skip": true, "comment": "Not supported" } ] },
+				{ "inputs": [ { "shape": [ 1, 2, 3 ] }, { "shape": [ 2 ] }, { "shape": [ 2 ] } ], "outputs": [ { "shape": [ 1, 2, 3 ] } ], "runtimes": [ { "name": "NNERuntimeORTDmlEditor", "skip": true, "comment": "Not supported" } ] },
+				{ "inputs": [ { "shape": [ 2, 2, 3 ] }, { "shape": [ 2 ] }, { "shape": [ 2 ] } ], "outputs": [ { "shape": [ 2, 2, 3 ] } ], "runtimes": [ { "name": "NNERuntimeORTDmlEditor", "skip": true, "comment": "Not supported" } ]},
 				{ "inputs": [ { "shape": [ 1, 2, 1, 3 ] }, { "shape": [ 2 ] }, { "shape": [ 2 ] } ], "outputs": [ { "shape": [ 1, 2, 1, 3 ] } ] },
 				{ "inputs": [ { "shape": [ 2, 2, 1, 3 ] }, { "shape": [ 2 ] }, { "shape": [ 2 ] } ], "outputs": [ { "shape": [ 2, 2, 1, 3 ] } ] },
 				{ "inputs": [ { "shape": [ 2, 2, 2, 2 ] }, { "shape": [ 2 ] }, { "shape": [ 2 ] } ], "outputs": [ { "shape": [ 2, 2, 2, 2 ] } ] },
@@ -665,64 +665,64 @@
 					"inputs": [ { "shape": [ 2, 2, 4, 1, 3 ] }, { "shape": [ 2 ] }, { "shape": [ 2 ] } ], "outputs": [ { "shape": [ 2, 2, 4, 1, 3 ] } ],
 					"runtimes": [
 						{ "name": "NNXRuntimeHLSL" },
-						{ "name": "NNXRuntimeORTCpu" },
-						{ "name": "NNXRuntimeORTDml", "skip": true, "comment": "Not supported" }
+						{ "name": "NNERuntimeORTCpuEditor" },
+						{ "name": "NNERuntimeORTDmlEditor", "skip": true, "comment": "Not supported" }
 					]
 				},
 				{
 					"inputs": [ { "shape": [ 2, 2, 4, 4 ] }, { "shape": [ 2 ] }, { "shape": [ 2 ] } ], "outputs": [ { "shape": [ 2, 2, 4, 4 ] } ],
 					"runtimes": [
 						{ "name": "NNXRuntimeHLSL" },
-						{ "name": "NNXRuntimeORTCpu" },
-						{ "name": "NNXRuntimeORTDml" }
+						{ "name": "NNERuntimeORTCpuEditor" },
+						{ "name": "NNERuntimeORTDmlEditor" }
 					]
 				},
 				{
 					"inputs": [ { "shape": [ 2, 2, 8, 8 ] }, { "shape": [ 2 ] }, { "shape": [ 2 ] } ], "outputs": [ { "shape": [ 2, 2, 8, 8 ] } ],
 					"runtimes": [
 						{ "name": "NNXRuntimeHLSL" },
-						{ "name": "NNXRuntimeORTCpu" },
-						{ "name": "NNXRuntimeORTDml" }
+						{ "name": "NNERuntimeORTCpuEditor" },
+						{ "name": "NNERuntimeORTDmlEditor" }
 					]
 				},
 				{
 					"inputs": [ { "shape": [ 2, 128, 8, 8 ] }, { "shape": [ 128 ] }, { "shape": [ 128 ] } ], "outputs": [ { "shape": [ 2, 128, 8, 8 ] } ],
 					"runtimes": [
 						{ "name": "NNXRuntimeHLSL" },
-						{ "name": "NNXRuntimeORTCpu" },
-						{ "name": "NNXRuntimeORTDml" }
+						{ "name": "NNERuntimeORTCpuEditor" },
+						{ "name": "NNERuntimeORTDmlEditor" }
 					]
 				},
 				{
 					"inputs": [ { "shape": [ 128, 2, 8, 8 ] }, { "shape": [ 2 ] }, { "shape": [ 2 ] } ], "outputs": [ { "shape": [ 128, 2, 8, 8 ] } ],
 					"runtimes": [
 						{ "name": "NNXRuntimeHLSL" },
-						{ "name": "NNXRuntimeORTCpu" },
-						{ "name": "NNXRuntimeORTDml" }
+						{ "name": "NNERuntimeORTCpuEditor" },
+						{ "name": "NNERuntimeORTDmlEditor" }
 					]
 				},
 				{
 					"inputs": [ { "shape": [ 64, 64, 8, 8 ] }, { "shape": [ 64 ] }, { "shape": [ 64 ] } ], "outputs": [ { "shape": [ 64, 64, 8, 8 ] } ],
 					"runtimes": [
 						{ "name": "NNXRuntimeHLSL" },
-						{ "name": "NNXRuntimeORTCpu" },
-						{ "name": "NNXRuntimeORTDml" }
+						{ "name": "NNERuntimeORTCpuEditor" },
+						{ "name": "NNERuntimeORTDmlEditor" }
 					]
 				},
 				{
 					"inputs": [ { "shape": [ 2, 2, 45, 45 ] }, { "shape": [ 2 ] }, { "shape": [ 2 ] } ], "outputs": [ { "shape": [ 2, 2, 45, 45 ] } ],
 					"runtimes": [
 						{ "name": "NNXRuntimeHLSL" },
-						{ "name": "NNXRuntimeORTCpu" },
-						{ "name": "NNXRuntimeORTDml" }
+						{ "name": "NNERuntimeORTCpuEditor" },
+						{ "name": "NNERuntimeORTDmlEditor" }
 					]
 				},
 				{
 					"inputs": [ { "shape": [ 1, 128, 56, 56 ] }, { "shape": [ 128 ] }, { "shape": [ 128 ] } ], "outputs": [ { "shape": [ 1, 128, 56, 56 ] } ],
 					"runtimes": [
 						{ "name": "NNXRuntimeHLSL" },
-						{ "name": "NNXRuntimeORTCpu" },
-						{ "name": "NNXRuntimeORTDml" }
+						{ "name": "NNERuntimeORTCpuEditor" },
+						{ "name": "NNERuntimeORTDmlEditor" }
 					],
 					"comment": "from 'mosaic-9-sim.onnx'"
 				}
diff --git a/Engine/Plugins/Experimental/NNE/Source/NNXRuntimeORT/NNXRuntimeORT.Build.cs b/Engine/Plugins/Experimental/NNE/Source/NNXRuntimeORT/NNXRuntimeORT.Build.cs
deleted file mode 100644
diff --git a/Engine/Plugins/Experimental/NNE/Source/NNXRuntimeORT/Private/NNXRuntimeORT.cpp b/Engine/Plugins/Experimental/NNE/Source/NNXRuntimeORT/Private/NNXRuntimeORT.cpp
deleted file mode 100644
diff --git a/Engine/Plugins/Experimental/NNE/Source/NNXRuntimeORT/Private/NNXRuntimeORT.h b/Engine/Plugins/Experimental/NNE/Source/NNXRuntimeORT/Private/NNXRuntimeORT.h
deleted file mode 100644
diff --git a/Engine/Plugins/Experimental/NNE/Source/NNXRuntimeORT/Private/NNXRuntimeORTModule.cpp b/Engine/Plugins/Experimental/NNE/Source/NNXRuntimeORT/Private/NNXRuntimeORTModule.cpp
deleted file mode 100644
diff --git a/Engine/Plugins/Experimental/NNE/Source/NNXRuntimeORT/Private/NNXRuntimeORTProviders.h b/Engine/Plugins/Experimental/NNE/Source/NNXRuntimeORT/Private/NNXRuntimeORTProviders.h
deleted file mode 100644
diff --git a/Engine/Plugins/Experimental/NNE/Source/NNXRuntimeORT/Private/NNXRuntimeORTUtils.h b/Engine/Plugins/Experimental/NNE/Source/NNXRuntimeORT/Private/NNXRuntimeORTUtils.h
deleted file mode 100644
