NNE Fix for Issue 251834: NonUnity errors in Analysis.h and UniquePtr.h

diff --git a/Engine/Plugins/Experimental/NNE/Source/NNERuntimeORT/Private/NNERuntimeORT.cpp b/Engine/Plugins/Experimental/NNE/Source/NNERuntimeORT/Private/NNERuntimeORT.cpp
--- a/Engine/Plugins/Experimental/NNE/Source/NNERuntimeORT/Private/NNERuntimeORT.cpp
+++ b/Engine/Plugins/Experimental/NNE/Source/NNERuntimeORT/Private/NNERuntimeORT.cpp
@@ -10,13 +10,6 @@
 #include "NNEProfilingTimer.h"
 #include "RedirectCoutAndCerrToUeLog.h"
 
-#include "NNXThirdPartyWarningDisabler.h"
-NNX_THIRD_PARTY_INCLUDES_START
-#undef check
-#undef TEXT
-#include "core/session/onnxruntime_cxx_api.h"
-NNX_THIRD_PARTY_INCLUDES_END
-
 #include UE_INLINE_GENERATED_CPP_BY_NAME(NNERuntimeORT)
 
 using namespace NNX;
diff --git a/Engine/Plugins/Experimental/NNE/Source/NNERuntimeORT/Private/NNERuntimeORT.h b/Engine/Plugins/Experimental/NNE/Source/NNERuntimeORT/Private/NNERuntimeORT.h
--- a/Engine/Plugins/Experimental/NNE/Source/NNERuntimeORT/Private/NNERuntimeORT.h
+++ b/Engine/Plugins/Experimental/NNE/Source/NNERuntimeORT/Private/NNERuntimeORT.h
@@ -10,9 +10,14 @@
 #include "UObject/Object.h"
 #include "UObject/UObjectBaseUtility.h"
 
-#include "NNERuntimeORT.generated.h"
+#include "NNXThirdPartyWarningDisabler.h"
+NNX_THIRD_PARTY_INCLUDES_START
+#undef check
+#undef TEXT
+#include "core/session/onnxruntime_cxx_api.h"
+NNX_THIRD_PARTY_INCLUDES_END
 
-namespace Ort { struct Env; }
+#include "NNERuntimeORT.generated.h"
 
 UCLASS()
 class UNNERuntimeORTCpuImpl : public UObject, public INNERuntime, public INNERuntimeCPU
