NNE : Minimal API - NNEQA

Problem:
* NNE API surface is large atm potentially creating a support burden and harming adoption for 5.2.

Solution:
* This CL is targeting NNEQA, removing any API from it.
* cleanup/rename module/file/classes to follow naming convention.
* bonus: update test definition so smoketest (nne.test.smoketest) are back under a minute.

Test:
* All tests green under win64 editor + style transfer demo in editor

diff --git a/Engine/Plugins/Experimental/NNE/NNE.uplugin b/Engine/Plugins/Experimental/NNE/NNE.uplugin
--- a/Engine/Plugins/Experimental/NNE/NNE.uplugin
+++ b/Engine/Plugins/Experimental/NNE/NNE.uplugin
@@ -44,7 +44,7 @@
 			"PlatformAllowList": [ "Win64", "Linux", "Mac" ]
 		},
 		{
-			"Name": "NNXQA",
+			"Name": "NNEQA",
 			"Type": "RuntimeAndProgram",
 			"LoadingPhase": "Default",
 			"PlatformAllowList": [ "Win64", "Linux", "Mac" ]
diff --git a/Engine/Plugins/Experimental/NNE/Source/NNEQA/NNEQA.Build.cs b/Engine/Plugins/Experimental/NNE/Source/NNEQA/NNEQA.Build.cs
new file mode 100644
--- /dev/null
+++ b/Engine/Plugins/Experimental/NNE/Source/NNEQA/NNEQA.Build.cs
@@ -0,0 +1,36 @@
+// Copyright Epic Games, Inc. All Rights Reserved.
+
+using System.IO;
+using UnrealBuildTool;
+
+public class NNEQA : ModuleRules
+{
+	public NNEQA(ReadOnlyTargetRules Target) : base(Target)
+	{
+		PCHUsage = ModuleRules.PCHUsageMode.UseExplicitOrSharedPCHs;
+
+		PublicDependencyModuleNames.AddRange(
+			new string[] {
+				"Core",
+			}
+		);
+
+		PrivateDependencyModuleNames.AddRange(
+			new string[]
+			{
+				"Engine",
+				"Projects",
+				"Json",
+				"CoreUObject",
+				"NNXCore",
+				"NNEUtils"
+			}
+		);
+
+		// RuntimeDependencies
+		foreach (string JsonFilePath in Directory.EnumerateFiles(Path.Combine(ModuleDirectory, "Resources"), "*.json"))
+		{
+			RuntimeDependencies.Add(JsonFilePath);
+		}
+	}
+}
diff --git a/Engine/Plugins/Experimental/NNE/Source/NNEQA/Private/NNEQAJsonUtils.h b/Engine/Plugins/Experimental/NNE/Source/NNEQA/Private/NNEQAJsonUtils.h
new file mode 100644
--- /dev/null
+++ b/Engine/Plugins/Experimental/NNE/Source/NNEQA/Private/NNEQAJsonUtils.h
@@ -0,0 +1,265 @@
+// Copyright Epic Games, Inc. All Rights Reserved.
+
+#pragma once
+
+#include "Serialization/JsonReader.h"
+#include "Serialization/JsonSerializer.h"
+#include "Serialization/JsonSerializerMacros.h"
+#include "Dom/JsonObject.h"
+#include "Misc/FileHelper.h"
+
+namespace UE::NNEQA::Private
+{
+namespace Json
+{
+	constexpr float JSON_TOLERANCE_NOTSET = -1.0f; 
+
+	struct FTestConfigTensor : FJsonSerializable
+	{
+		TArray<int32> Shape;
+		TArray<FString> Source;
+		FString Type;
+		//Idea: Extend with custom initializer:
+		//	FString Initializer;
+		//	JSON_SERIALIZE("initializer", ...);
+		//		if (initializer && source)
+		//			InEngineTensorData = initializer(source) <-- To add
+		//		if (!initializer && source)
+		//			InEngineTensorData = element-wise( static_cast<Type>(source[i]) ) <-- Current behavior
+		//		if (!initializer && !source)
+		//			InEngineTensorData = random from seed <-- Current behavior
+		//Idea: Extend by adding source CPU or GPU resident data
+		//	bool OnGpu;
+		//	JSON_SERIALIZE("on_gpu", OnGpu);
+		BEGIN_JSON_SERIALIZER
+			JSON_SERIALIZE_ARRAY("shape", Shape);
+			JSON_SERIALIZE("type", Type);
+			JSON_SERIALIZE_ARRAY("source", Source);
+		END_JSON_SERIALIZER
+	};
+
+	struct FTestConfigRuntime : FJsonSerializable
+	{
+		FString Name;
+		bool Skip;
+		float AbsoluteTolerance = JSON_TOLERANCE_NOTSET;
+		float RelativeTolerance = JSON_TOLERANCE_NOTSET;
+
+		BEGIN_JSON_SERIALIZER
+			JSON_SERIALIZE("name", Name);
+			JSON_SERIALIZE("skip", Skip);
+			JSON_SERIALIZE("absolute_tolerance", AbsoluteTolerance);
+			JSON_SERIALIZE("relative_tolerance", RelativeTolerance);
+		END_JSON_SERIALIZER
+	};
+
+	struct FTestConfigDataset : FJsonSerializable
+	{
+		TArray<FTestConfigTensor> Inputs;
+		TArray<FTestConfigTensor> Weights;
+		TArray<FTestConfigTensor> Outputs;
+		TArray<FTestConfigRuntime> Runtimes;
+
+		BEGIN_JSON_SERIALIZER
+			JSON_SERIALIZE_ARRAY_SERIALIZABLE("inputs", Inputs, FTestConfigTensor);
+			JSON_SERIALIZE_ARRAY_SERIALIZABLE("weights", Weights, FTestConfigTensor);
+			JSON_SERIALIZE_ARRAY_SERIALIZABLE("outputs", Outputs, FTestConfigTensor);
+			JSON_SERIALIZE_ARRAY_SERIALIZABLE("runtimes", Runtimes, FTestConfigRuntime);
+		END_JSON_SERIALIZER
+	};
+
+	struct FTestAttribute : FJsonSerializable
+	{
+		FString Name;
+		FNNEAttributeValue Value;
+
+		BEGIN_JSON_SERIALIZER
+			JSON_SERIALIZE("name", Name);
+			
+			// Since we prefer the attribute value to be in readable form in the Json and
+			// we only serialize it, we construct it at runtime by manually serialize it
+			// note: default type is float
+
+			FString TypeStr;
+			Serializer.Serialize(TEXT("type"), TypeStr);
+
+			ENNEAttributeDataType Type = ENNEAttributeDataType::Float;
+			if (!TypeStr.IsEmpty()) LexFromString(Type, *TypeStr);
+
+			switch (Type)
+			{
+				case ENNEAttributeDataType::Float:
+				{
+					float TmpValue;
+					JSON_SERIALIZE("value", TmpValue);
+					Value = FNNEAttributeValue(TmpValue);
+					break;
+				}
+
+				case ENNEAttributeDataType::Int32:
+				{
+					int TmpValue;
+					JSON_SERIALIZE("value", TmpValue);
+					Value = FNNEAttributeValue(TmpValue);
+					break;
+				}
+				case ENNEAttributeDataType::Int32Array:
+				{
+					TArray<int32> TmpValue;
+					JSON_SERIALIZE_ARRAY("value", TmpValue);
+					Value = FNNEAttributeValue(TmpValue);
+					break;
+				}
+				case ENNEAttributeDataType::String:
+				{
+					FString TmpValue;
+					JSON_SERIALIZE("value", TmpValue);
+					Value = FNNEAttributeValue(TmpValue);
+					break;
+				}
+				default:
+					check(Type == ENNEAttributeDataType::None);
+			}
+		END_JSON_SERIALIZER
+	};
+
+	struct FTestAttributeMap : FJsonSerializable
+	{
+		TArray<FTestAttribute> Attributes;
+
+		BEGIN_JSON_SERIALIZER
+			JSON_SERIALIZE_ARRAY_SERIALIZABLE("attributes", Attributes, FTestAttribute);
+		END_JSON_SERIALIZER
+	};
+
+	struct FTestConfigTarget : FJsonSerializable
+	{
+		FString Target;
+		TArray<FString> Tags;
+		TArray<FString> AdditionalDatasets;
+		TArray<FString> RemovedDatasets;
+		bool Skip;
+		float AbsoluteTolerance = JSON_TOLERANCE_NOTSET;
+		float RelativeTolerance = JSON_TOLERANCE_NOTSET;
+		FString InputType;
+		FString OutputType;
+		TArray<FTestConfigRuntime> Runtimes;
+
+		BEGIN_JSON_SERIALIZER
+			JSON_SERIALIZE("target", Target);
+			JSON_SERIALIZE("skip", Skip);
+			JSON_SERIALIZE("absolute_tolerance", AbsoluteTolerance);
+			JSON_SERIALIZE("relative_tolerance", RelativeTolerance);
+			JSON_SERIALIZE("input_type", InputType);
+			JSON_SERIALIZE("output_type", OutputType);
+			JSON_SERIALIZE_ARRAY_SERIALIZABLE("runtimes", Runtimes, FTestConfigRuntime);
+			JSON_SERIALIZE_ARRAY("tags", Tags);
+			JSON_SERIALIZE_ARRAY("additional_datasets", AdditionalDatasets);
+			JSON_SERIALIZE_ARRAY("removed_datasets", RemovedDatasets);
+		END_JSON_SERIALIZER
+	};
+
+	struct FTestCategory : FJsonSerializable
+	{
+		FString Category;
+		TArray<FString> AdditionalDatasets;
+		TArray<FString> RemovedDatasets;
+		TArray<FTestConfigTarget> Targets;
+		TArray<FTestConfigRuntime> Runtimes;
+		bool Skip;
+		bool IsModelTest;
+
+		BEGIN_JSON_SERIALIZER
+			JSON_SERIALIZE("category", Category);
+			JSON_SERIALIZE("skip", Skip);
+			JSON_SERIALIZE("is_model_test", IsModelTest);
+			JSON_SERIALIZE_ARRAY("additional_datasets", AdditionalDatasets);
+			JSON_SERIALIZE_ARRAY("removed_datasets", RemovedDatasets);
+			JSON_SERIALIZE_ARRAY_SERIALIZABLE("targets", Targets, FTestConfigTarget);
+			JSON_SERIALIZE_ARRAY_SERIALIZABLE("runtimes", Runtimes, FTestConfigRuntime);
+		END_JSON_SERIALIZER
+	};
+
+	struct FTestConfigInputOutputSet : FJsonSerializable
+	{
+		FString Name;
+		TArray<FTestConfigDataset> Datasets;
+		TArray<FTestConfigRuntime> Runtimes;
+
+		BEGIN_JSON_SERIALIZER
+			JSON_SERIALIZE("name", Name);
+			JSON_SERIALIZE_ARRAY_SERIALIZABLE("datasets", Datasets, FTestConfigDataset);
+			JSON_SERIALIZE_ARRAY_SERIALIZABLE("runtimes", Runtimes, FTestConfigRuntime);
+		END_JSON_SERIALIZER
+	};
+
+	struct FTestAttributeSet : FJsonSerializable
+	{
+		FString Name;
+		TArray<FTestAttributeMap> AttributeMaps;
+		TArray<FString> MultiplyWithAttributeSets;
+
+		BEGIN_JSON_SERIALIZER
+			JSON_SERIALIZE("name", Name);
+			JSON_SERIALIZE_ARRAY_SERIALIZABLE("attribute_maps", AttributeMaps, FTestAttributeMap);
+			JSON_SERIALIZE_ARRAY("multiply_with_attribute_sets", MultiplyWithAttributeSets);
+		END_JSON_SERIALIZER
+	};
+
+	template<typename T>
+	bool TrySerializeArray(TArray<T> &OutArray, const FString &FieldName, const FJsonObject &GlobalJsonObject)
+	{
+		const TArray<TSharedPtr<FJsonValue>>* JsonArrayValue;
+		if (GlobalJsonObject.TryGetArrayField(FieldName, JsonArrayValue))
+		{
+			for (TSharedPtr<FJsonValue> JsonValue : *JsonArrayValue)
+			{
+				const TSharedPtr<FJsonObject>* JsonObject = nullptr;
+				if (!JsonValue->TryGetObject(JsonObject))
+				{
+					return false;
+				}
+				OutArray.Emplace_GetRef().FromJson(*JsonObject);
+			}
+		}
+
+		return true;
+	}
+
+	bool LoadTestDescriptionFromJson(const FString& FullPath,
+		TArray<FTestCategory>& ModelTestCategories,
+		TArray<FTestCategory>& OperatorCategories,
+		TArray<FTestConfigInputOutputSet>& InputOutputSets,
+		TArray<FTestAttributeSet>& AttributeSets)
+	{
+		ModelTestCategories.Empty();
+		OperatorCategories.Empty();
+		InputOutputSets.Empty();
+		AttributeSets.Empty();
+
+		FString JsonContent;
+		FFileHelper::LoadFileToString(JsonContent, *FullPath);
+		TSharedRef<TJsonReader<TCHAR>> JsonReader = TJsonReaderFactory<TCHAR>::Create(JsonContent);
+		TSharedPtr<FJsonObject> GlobalJsonObject = MakeShareable(new FJsonObject);
+
+		if (!FJsonSerializer::Deserialize(JsonReader, GlobalJsonObject))
+		{
+			return false;
+		}
+
+		if (!GlobalJsonObject.IsValid())
+		{
+			return false;
+		}
+
+		bool bSuccess = false;
+		bSuccess |= TrySerializeArray(ModelTestCategories, TEXT("model_test_categories"), *GlobalJsonObject);
+		bSuccess |= TrySerializeArray(OperatorCategories, TEXT("operator_test_categories"), *GlobalJsonObject);
+		bSuccess |= TrySerializeArray(InputOutputSets, TEXT("input_output_sets"), *GlobalJsonObject);
+		bSuccess |= TrySerializeArray(AttributeSets, TEXT("attribute_sets"), *GlobalJsonObject);
+
+		return bSuccess;
+	}
+
+} // namespace Json
+} // namespace UE::NNEQA::Private
diff --git a/Engine/Plugins/Experimental/NNE/Source/NNEQA/Private/NNEQAModule.cpp b/Engine/Plugins/Experimental/NNE/Source/NNEQA/Private/NNEQAModule.cpp
new file mode 100644
--- /dev/null
+++ b/Engine/Plugins/Experimental/NNE/Source/NNEQA/Private/NNEQAModule.cpp
@@ -0,0 +1,22 @@
+// Copyright Epic Games, Inc. All Rights Reserved.
+
+#include "Modules/ModuleManager.h"
+
+#include "NNEQAParametricTest.h"
+
+class FNNEQAModule : public IModuleInterface
+{
+public:
+
+	virtual void StartupModule() override
+	{
+		UE::NNEQA::Private::InitializeParametricTests();
+	}
+
+	virtual void ShutdownModule() override
+	{
+	}
+
+};
+
+IMPLEMENT_MODULE(FNNEQAModule, NNEQA);
\ No newline at end of file
diff --git a/Engine/Plugins/Experimental/NNE/Source/NNEQA/Private/NNEQAParametricTest.cpp b/Engine/Plugins/Experimental/NNE/Source/NNEQA/Private/NNEQAParametricTest.cpp
new file mode 100644
--- /dev/null
+++ b/Engine/Plugins/Experimental/NNE/Source/NNEQA/Private/NNEQAParametricTest.cpp
@@ -0,0 +1,718 @@
+// Copyright Epic Games, Inc. All Rights Reserved.
+
+#include "NNEQAParametricTest.h"
+
+#include "HAL/IConsoleManager.h"
+#include "Interfaces/IPluginManager.h"
+#include "Misc/Paths.h"
+#include "NNECore.h"
+#include "NNECoreTypes.h"
+#include "NNECoreAttributeMap.h"
+#include "NNEQAUtils.h"
+#include "NNEQAJsonUtils.h"
+#include "NNEUtilsModelBuilder.h"
+#include "UObject/Class.h"
+#include "UObject/ReflectedTypeAccessors.h"
+
+#if WITH_DEV_AUTOMATION_TESTS
+#include "Misc/AutomationTest.h"
+#endif
+
+namespace UE::NNEQA::Private
+{
+	
+class FParametricTests : public FTests
+{
+public:
+	FParametricTests() = default;
+
+	bool ReloadTestDescriptionsFromJson()
+	{
+		TestSetups.Empty();
+
+		TSharedPtr<IPlugin> NNEPlugin = IPluginManager::Get().FindPlugin(TEXT("NNE"));
+			
+		if (!NNEPlugin.IsValid())
+		{
+			UE_LOG(LogNNE, Error, TEXT("Unable to find NNE plugin!"));
+			return false;
+		}
+
+		//Notes: We should allow to define the tests in more than one json file
+		FString NNEPluginBaseDir = NNEPlugin->GetBaseDir();
+		FString FullPath(NNEPluginBaseDir + TEXT("\\Source\\NNEQA\\Resources\\NNEQATestDesc.json"));
+
+		TArray<Json::FTestCategory> ModelTestCategories;
+		TArray<Json::FTestCategory> OperatorTestCategories;
+		TArray<Json::FTestConfigInputOutputSet> InputOutputSets;
+		TArray<Json::FTestAttributeSet> AttributeSets;
+			
+		if (!Json::LoadTestDescriptionFromJson(FullPath, ModelTestCategories, OperatorTestCategories, InputOutputSets, AttributeSets))
+		{
+			return false;
+		}
+
+		const FString NNEBaseTestPath(TEXT("System.Engine.MachineLearning.NNE"));
+		AddTestFromCategory(NNEBaseTestPath + TEXT(".Model."), ModelTestCategories, InputOutputSets, AttributeSets);
+		AddTestFromCategory(NNEBaseTestPath + TEXT(".Operator."), OperatorTestCategories, InputOutputSets, AttributeSets);
+
+		return true;
+	}
+private:
+		
+	static void ApplyEpsilons(FTests::FTestSetup& TestSetup, const Json::FTestConfigTarget& TestTarget)
+	{
+		if (TestTarget.AbsoluteTolerance != Json::JSON_TOLERANCE_NOTSET)
+		{
+			TestSetup.AbsoluteTolerance = TestTarget.AbsoluteTolerance;
+		}
+			
+		if (TestTarget.RelativeTolerance != Json::JSON_TOLERANCE_NOTSET)
+		{
+			TestSetup.RelativeTolerance = TestTarget.RelativeTolerance;
+		}
+	}
+
+	static void ApplyRuntimesConfig(FTests::FTestSetup& TestSetup, const TArray<Json::FTestConfigRuntime>& TestRuntimes)
+	{
+		for (const Json::FTestConfigRuntime& Runtime : TestRuntimes)
+		{
+			if (Runtime.Skip && !TestSetup.AutomationExcludedRuntime.Contains(Runtime.Name))
+			{
+				TestSetup.AutomationExcludedRuntime.Emplace(Runtime.Name);
+			}
+			else
+			{
+				if (Runtime.AbsoluteTolerance != Json::JSON_TOLERANCE_NOTSET)
+				{
+					TestSetup.AbsoluteToleranceForRuntime.Emplace(Runtime.Name, Runtime.AbsoluteTolerance);
+				}
+				if (Runtime.RelativeTolerance != Json::JSON_TOLERANCE_NOTSET)
+				{
+					TestSetup.RelativeToleranceForRuntime.Emplace(Runtime.Name, Runtime.RelativeTolerance);
+				}
+			}
+		}
+	}
+
+	static void ApplyTargetConfig(FTests::FTestSetup& TestSetup, const Json::FTestConfigTarget& TestTarget)
+	{
+		//Epsilons
+		ApplyEpsilons(TestSetup, TestTarget);
+
+		ApplyRuntimesConfig(TestSetup, TestTarget.Runtimes);
+
+		//Tags
+		TestSetup.Tags = TestTarget.Tags;
+	}
+
+	static NNECore::FTensorShape GetConcreteShapeFromJsonArray(TConstArrayView<int32> JsonShape)
+	{
+		NNECore::FSymbolicTensorShape SymbolicShape = NNECore::FSymbolicTensorShape::Make(JsonShape);
+		check(SymbolicShape.IsConcrete());
+		NNECore::FTensorShape ConcreteShape = NNECore::FTensorShape::MakeFromSymbolic(SymbolicShape);
+		return ConcreteShape;
+	}
+
+	static ENNETensorDataType GetTensorTypeFromJson(const FString& TypeName, ENNETensorDataType DefaultValue)
+	{
+		int64 Value = StaticEnum<ENNETensorDataType>()->GetValueByNameString(TypeName);
+			
+		return (Value == INDEX_NONE) ? DefaultValue : (ENNETensorDataType)Value;
+	}
+
+	class ElementWiseFromJsonStringTensorInitializer
+	{
+		TConstArrayView<FString> JsonValues;
+				
+	public:
+		ElementWiseFromJsonStringTensorInitializer(TConstArrayView<FString> InJsonValues)
+			:JsonValues(InJsonValues) 
+		{}
+
+		float operator () (uint32 TensorIndex) const 
+		{
+			const FString& JsonValue = JsonValues[TensorIndex];
+			if (!FCString::IsNumeric(*JsonValue))
+			{
+				UE_LOG(LogNNE, Error, TEXT("Cannot convert %s to float will default to 0.0f, check test config."), *JsonValue);
+			}
+			return FCString::Atof(*JsonValues[TensorIndex]);
+		}
+	};
+
+	static TArray<char> GetInputTensorDataFromJson(const NNECore::Internal::FTensor& Tensor, int TensorIndex, TConstArrayView<FString> JsonValues)
+	{
+		if (JsonValues.Num() == Tensor.GetVolume())
+		{
+			ElementWiseFromJsonStringTensorInitializer Initializer(JsonValues);
+			return GenerateTensorDataForTest(Tensor, Initializer);
+		}
+		else if (JsonValues.Num() > 0)
+		{
+			UE_LOG(LogNNE, Error, TEXT("Incorrect number of element for tensor initializer %s, should have %d but got %d. Fallbacking to default initializer."), *Tensor.GetName(), Tensor.GetVolume(), JsonValues.Num());
+				
+		}
+		ElementWiseCosTensorInitializer Initializer(Tensor.GetDataType(), TensorIndex);
+		return GenerateTensorDataForTest(Tensor, Initializer);
+	}
+
+	static TArray<char> GetOutputTensorDataFromJson(const NNECore::Internal::FTensor& Tensor, int TensorIndex, TConstArrayView<FString> JsonValues)
+	{
+		if (JsonValues.Num() == Tensor.GetVolume())
+		{
+			ElementWiseFromJsonStringTensorInitializer Initializer(JsonValues);
+			return GenerateTensorDataForTest(Tensor, Initializer);
+		}
+		else if (JsonValues.Num() > 0)
+		{
+			UE_LOG(LogNNE, Error, TEXT("Incorrect number of element for tensor initializer %s, should have %d but got %d."), *Tensor.GetName());
+
+		}
+		return TArray<char>();
+	}
+
+	static void ApplyDatasetConfig(FTests::FTestSetup& TestSetup, const Json::FTestConfigDataset& TestDataset, ENNETensorDataType DefaultInputType, ENNETensorDataType DefaultOutputType)
+	{
+		ApplyRuntimesConfig(TestSetup, TestDataset.Runtimes);
+
+		if (TestDataset.Inputs.Num() == 0)
+		{
+			return;
+		}
+
+		uint32 i = 0;
+		for (auto&& JsonTensor : TestDataset.Inputs)
+		{
+			NNECore::FTensorShape Shape = GetConcreteShapeFromJsonArray(JsonTensor.Shape);
+			ENNETensorDataType TensorType = GetTensorTypeFromJson(JsonTensor.Type, DefaultInputType);
+			NNECore::Internal::FTensor Tensor = NNECore::Internal::FTensor::Make(FString::Printf(TEXT("input%d"), i), Shape, TensorType);
+			TArray<char> TensorData = GetInputTensorDataFromJson(Tensor, i, JsonTensor.Source);
+
+			TestSetup.Inputs.Emplace(Tensor);
+			TestSetup.InputsData.Emplace(TensorData);
+			++i;
+		}
+
+		for (auto&& JsonTensor : TestDataset.Weights)
+		{
+			NNECore::FTensorShape Shape = GetConcreteShapeFromJsonArray(JsonTensor.Shape);
+			ENNETensorDataType TensorType = GetTensorTypeFromJson(JsonTensor.Type, ENNETensorDataType::Float);
+			NNECore::Internal::FTensor Tensor = NNECore::Internal::FTensor::Make(FString::Printf(TEXT("weights%d"), i), Shape, TensorType);
+			TArray<char> TensorData = GetInputTensorDataFromJson(Tensor, i, JsonTensor.Source);
+
+			TestSetup.Weights.Emplace(Tensor);
+			TestSetup.WeightsData.Emplace(TensorData);
+			++i;
+		}
+
+		for (auto&& JsonTensor : TestDataset.Outputs)
+		{
+			NNECore::FTensorShape Shape = GetConcreteShapeFromJsonArray(JsonTensor.Shape);
+			ENNETensorDataType TensorType = GetTensorTypeFromJson(JsonTensor.Type, DefaultOutputType);
+			NNECore::Internal::FTensor Tensor = NNECore::Internal::FTensor::Make(FString::Printf(TEXT("output%d"), i), Shape, TensorType);
+			TArray<char> TensorData = GetOutputTensorDataFromJson(Tensor, i, JsonTensor.Source);
+
+			TestSetup.Outputs.Emplace(Tensor);
+			TestSetup.OutputsData.Emplace(TensorData);
+			++i;
+		}
+		//If output is not defined it is the first input shape.
+		if (TestDataset.Outputs.Num() == 0 && TestDataset.Inputs.Num() > 0)
+		{
+			NNECore::FTensorShape Shape = GetConcreteShapeFromJsonArray(TestDataset.Inputs[0].Shape);
+			ENNETensorDataType TensorType = GetTensorTypeFromJson(TestDataset.Inputs[0].Type, DefaultOutputType);
+			NNECore::Internal::FTensor Tensor = NNECore::Internal::FTensor::Make(FString::Printf(TEXT("output"), i), Shape, TensorType);
+			TArray<char> TensorData = GetOutputTensorDataFromJson(Tensor, i, TestDataset.Inputs[0].Source);
+				
+			TestSetup.Outputs.Emplace(Tensor);
+			TestSetup.OutputsData.Emplace(TensorData);
+		}
+	}
+
+	static void ApplyAttributeSetConfig(FTests::FTestSetup& TestSetup, const Json::FTestAttributeMap& AttributeMap)
+	{
+		for (const Json::FTestAttribute &Attribute : AttributeMap.Attributes)
+		{
+			TestSetup.AttributeMap.SetAttribute(Attribute.Name, Attribute.Value);
+		}
+	}
+
+	static FString GetTestSuffix(const Json::FTestConfigDataset& Dataset)
+	{
+		//Build TestSuffix "<inputshape0>_..._w<weightshape0>_w...=><outputshape0>_..."
+		bool bIsFirstShape = true;
+		FString TestSuffix;
+		for (const Json::FTestConfigTensor& Input : Dataset.Inputs)
+		{
+			if (!bIsFirstShape) TestSuffix += TEXT("_");
+			TestSuffix += ShapeToString<int32>(Input.Shape);
+			bIsFirstShape = false;
+		}
+		for (const Json::FTestConfigTensor& Weight : Dataset.Weights)
+		{
+			TestSuffix += TEXT("_w");
+			TestSuffix += ShapeToString<int32>(Weight.Shape);
+		}
+		TestSuffix += TEXT("=>");
+		//If output is not defined it is the first input shape.
+		if (Dataset.Outputs.Num() == 0)
+		{
+			TestSuffix += ShapeToString<int32>(Dataset.Inputs[0].Shape);
+		}
+		else
+		{
+			bIsFirstShape = true;
+			for (const Json::FTestConfigTensor& Output : Dataset.Outputs)
+			{
+				if (!bIsFirstShape) TestSuffix += TEXT("_");
+				TestSuffix += ShapeToString<int32>(Output.Shape);
+				bIsFirstShape = false;
+			}
+		}
+		return TestSuffix;
+	}
+
+	static FString GetAttributeMapSuffix(const Json::FTestAttributeMap &AttributeMap)
+	{
+		if (AttributeMap.Attributes.IsEmpty()) return TEXT("");
+
+		auto AttributeToString = [] (const FNNEAttributeValue &Value) -> FString {
+			ENNEAttributeDataType Type = Value.GetType();
+				
+			switch(Type)
+			{
+				case ENNEAttributeDataType::Float: return FString::SanitizeFloat(Value.GetValue<float>());
+				case ENNEAttributeDataType::Int32: return FString::FromInt(Value.GetValue<int32>());
+				case ENNEAttributeDataType::Int32Array:
+				{
+					TArray<FString> ArrayStrings;
+					for (int32 Val : Value.GetValue<TArray<int32>>())
+					{
+						ArrayStrings.Add(FString::FromInt(Val));
+					}
+
+					return TEXT("[") + FString::Join(ArrayStrings, TEXT(",")) + TEXT("]");
+				}
+				case ENNEAttributeDataType::String: return Value.GetValue<FString>();
+			}
+			return TEXT("-");
+		};
+			
+		FStringBuilderBase Str;
+		Str  << TEXT(".");
+
+		TArray<FString> AttributeStrings;
+		for (const Json::FTestAttribute &Attribute : AttributeMap.Attributes)
+		{
+			AttributeStrings.Push(Attribute.Name + "=" + AttributeToString(Attribute.Value));
+		}
+
+		Str  << FString::Join(AttributeStrings, TEXT(","));
+
+		return Str.ToString();
+	}
+
+	static bool IsSubstringFoundInArray(const TArray<FString>& Names, const FString& SubString)
+	{
+		for (auto&& Name : Names)
+		{
+			if (Name.Contains(SubString))
+			{
+				return true;
+			}
+		}
+		return false;
+	}
+		
+	void AddTestFromCategory(const FString& BaseTestPath, const TArray<Json::FTestCategory>& TestCategories, const TArray<Json::FTestConfigInputOutputSet>& InputOutputSets,
+			const TArray<Json::FTestAttributeSet>& AttributeSets)
+	{
+		for (const Json::FTestCategory& TestCategory : TestCategories)
+		{
+			if (TestCategory.Skip)
+			{
+				continue;
+			}
+
+			const bool bIsModelCategory = TestCategory.IsModelTest;
+			const FString TestCategoryPath(BaseTestPath + TestCategory.Category + TEXT("."));
+				
+			for (const Json::FTestConfigTarget& TestTarget : TestCategory.Targets)
+			{
+				if (TestTarget.Skip)
+				{
+					continue;
+				}
+
+				const FString& TestBaseName = TestTarget.Target;
+				ENNETensorDataType InputTypeFromTarget = GetTensorTypeFromJson(TestTarget.InputType, ENNETensorDataType::Float);
+				ENNETensorDataType OutputTypeFromTarget = GetTensorTypeFromJson(TestTarget.OutputType, ENNETensorDataType::Float);
+				bool bAtLeastATestWasAdded = false;
+					
+				for (const Json::FTestConfigInputOutputSet& InputOutputSet : InputOutputSets)
+				{
+					//Accept test if category is a substring of the dataset name or one of the explicitly requested dataset names
+					if (InputOutputSet.Name.Contains(TestCategory.Category) || 
+						IsSubstringFoundInArray(TestTarget.AdditionalDatasets, InputOutputSet.Name) ||
+						IsSubstringFoundInArray(TestCategory.AdditionalDatasets, InputOutputSet.Name))
+					{
+						//Reject test if dataset is explicitly rejected
+						if (IsSubstringFoundInArray(TestTarget.RemovedDatasets, InputOutputSet.Name) || 
+							IsSubstringFoundInArray(TestCategory.RemovedDatasets, InputOutputSet.Name))
+						{
+							continue;
+						}
+							
+						for (const Json::FTestConfigDataset& Dataset : InputOutputSet.Datasets)
+						{
+							if (Dataset.Inputs.Num() == 0)
+							{
+								continue;
+							}
+
+							bool bAtLeastAnAttributeTestWasWadded = false;
+							for (const Json::FTestAttributeSet &AttributeSet : AttributeSets)
+							{
+								// Idea: maybe split by '.' and (partially) match parts
+								if (!AttributeSet.Name.Contains(InputOutputSet.Name))
+								{
+									continue;
+								}
+
+								for (const Json::FTestAttributeMap &AttributeMap : AttributeSet.AttributeMaps)
+								{
+									bool bAtLeastAnotherAttributeSetAdded = false;
+									for (const FString &OtherAttributeSetName : AttributeSet.MultiplyWithAttributeSets)
+									{
+										for (const Json::FTestAttributeSet &OtherAttributeSet : AttributeSets)
+										{
+											if (OtherAttributeSet.Name != OtherAttributeSetName)
+											{
+												continue;
+											}
+
+											for (const Json::FTestAttributeMap &OtherAttributeMap : OtherAttributeSet.AttributeMaps)
+											{
+												FTests::FTestSetup& Test = AddTest(TestCategoryPath, TestBaseName, TEXT(".") + GetTestSuffix(Dataset) + GetAttributeMapSuffix(AttributeMap) + GetAttributeMapSuffix(OtherAttributeMap));
+										
+												ApplyRuntimesConfig(Test, TestCategory.Runtimes);
+												ApplyRuntimesConfig(Test, InputOutputSet.Runtimes);
+												ApplyTargetConfig(Test, TestTarget);
+												ApplyDatasetConfig(Test, Dataset, InputTypeFromTarget, OutputTypeFromTarget);
+
+												ApplyAttributeSetConfig(Test, AttributeMap);
+												ApplyAttributeSetConfig(Test, OtherAttributeMap);
+
+												Test.IsModelTest = bIsModelCategory;
+
+												bAtLeastAnAttributeTestWasWadded = true;
+
+												bAtLeastAnotherAttributeSetAdded = true;
+											}
+										}
+									}
+
+									if (!bAtLeastAnotherAttributeSetAdded)
+									{
+										FTests::FTestSetup& Test = AddTest(TestCategoryPath, TestBaseName, TEXT(".") + GetTestSuffix(Dataset) + GetAttributeMapSuffix(AttributeMap));
+											
+										ApplyRuntimesConfig(Test, TestCategory.Runtimes);
+										ApplyRuntimesConfig(Test, InputOutputSet.Runtimes);
+										ApplyTargetConfig(Test, TestTarget);
+										ApplyDatasetConfig(Test, Dataset, InputTypeFromTarget, OutputTypeFromTarget);
+
+										ApplyAttributeSetConfig(Test, AttributeMap);
+
+										Test.IsModelTest = bIsModelCategory;
+
+										bAtLeastAnAttributeTestWasWadded = true;
+									}
+								}
+							}
+
+							if (!bAtLeastAnAttributeTestWasWadded)
+							{
+							FTests::FTestSetup& Test = AddTest(TestCategoryPath, TestBaseName, TEXT(".") + GetTestSuffix(Dataset));
+								
+							ApplyRuntimesConfig(Test, TestCategory.Runtimes);
+							ApplyRuntimesConfig(Test, InputOutputSet.Runtimes);
+							ApplyTargetConfig(Test, TestTarget);
+							ApplyDatasetConfig(Test, Dataset, InputTypeFromTarget, OutputTypeFromTarget);
+							Test.IsModelTest = bIsModelCategory;
+						}
+						}
+						bAtLeastATestWasAdded = true;
+					}
+				}
+
+				//No dataset were matched with this target, define a test without enforcing input/output.
+				if (!bAtLeastATestWasAdded)
+				{
+					FTests::FTestSetup& Test = AddTest(TestCategoryPath, TestBaseName, TEXT(""));
+						
+					ApplyRuntimesConfig(Test, TestCategory.Runtimes);
+					ApplyTargetConfig(Test, TestTarget);
+					Test.IsModelTest = bIsModelCategory;
+				}
+			}
+		}
+	}
+};
+static FParametricTests ParametricTests;
+
+static FString GetFullModelPath(const FString& ModelName)
+{
+	//Note: This mean model tests can only run in the context of the current projects (example: NNXIncubator)
+	return FPaths::ConvertRelativePathToFull(FPaths::ProjectContentDir() / TEXT("OnnxModels") / ModelName);
+}
+	
+static bool RunParametricTest(FTests::FTestSetup& TestSetup, const FString& RuntimeFilter)
+{
+	FNNIModelRaw ONNXModel;
+	FNNIModelRaw ONNXModelVariadic;
+	ONNXModelVariadic.Format = ENNXInferenceFormat::Invalid;
+
+	if (TestSetup.IsModelTest)
+	{
+		// Model test, load model from disk
+		FString ModelPath = GetFullModelPath(TestSetup.TargetName + TEXT(".onnx"));
+		const bool bIsModelInMem = FFileHelper::LoadFileToArray(ONNXModel.Data, *ModelPath);
+		if (!bIsModelInMem)
+		{
+			UE_LOG(LogNNE, Error, TEXT("Can't load model from disk at path '%s'. Tests ABORTED!"), *ModelPath);
+			return false;
+		}
+		ONNXModel.Format = ENNXInferenceFormat::ONNX;
+	}
+	else
+	{
+		// Operator test, create model in memory
+		if (!UE::NNEUtils::Internal::CreateONNXModelForOperator(false, TestSetup.TargetName, TestSetup.Inputs, TestSetup.Outputs, TestSetup.Weights, TestSetup.WeightsData,TestSetup.AttributeMap, ONNXModel))
+		{
+			UE_LOG(LogNNE, Error, TEXT("Failed to create static model for test '%s'. Test ABORTED!"), *TestSetup.TargetName);
+			return false;
+		}
+		if (!UE::NNEUtils::Internal::CreateONNXModelForOperator(true, TestSetup.TargetName, TestSetup.Inputs, TestSetup.Outputs, TestSetup.Weights, TestSetup.WeightsData, TestSetup.AttributeMap, ONNXModelVariadic))
+		{
+			UE_LOG(LogNNE, Error, TEXT("Failed to create variadic model for test '%s'. Test ABORTED!"), *TestSetup.TargetName);
+			return false;
+		}
+
+        //FFileHelper::SaveArrayToFile(ONNXModel.Data, TEXT("D:\\Opmodel.onnx"));
+        //FFileHelper::SaveArrayToFile(ONNXModelVariadic.Data, TEXT("D:\\OpmodelVar.onnx"));
+	}
+
+	return CompareONNXModelInferenceAcrossRuntimes(ONNXModel, ONNXModelVariadic, TestSetup, RuntimeFilter);
+}
+
+static FString AutomationRuntimeFilter;
+void SetAutomationRuntimeFilter(const FString& InRuntimeFilter)
+{
+	AutomationRuntimeFilter = InRuntimeFilter;
+}
+
+static FAutoConsoleCommand SetAutomationRuntimeFilterCommand(
+	TEXT("nne.test.setautomationfilter"), TEXT("Set the RuntimeFilter witch automation will use, no parameter to run on all runtime (this is the default)."),
+	FConsoleCommandWithArgsDelegate::CreateStatic(
+		[](const TArray< FString >& Args)
+		{
+			AutomationRuntimeFilter.Empty();
+			for (auto&& Arg : Args)
+			{
+				AutomationRuntimeFilter += Arg;
+				AutomationRuntimeFilter += TEXT(" ");
+			}
+		}
+	)
+);
+
+bool RunParametricTests(const FString& NameSubstring, const FString& Tag, const FString& InRuntimeFilter)
+{
+	bool bAllTestSucceed = true;
+	uint32 NumTest = 0;
+	uint32 NumTestFailed = 0;
+		
+	for (FTests::FTestSetup& Test : ParametricTests.TestSetups)
+	{
+		if (!Tag.IsEmpty() && !Test.Tags.Contains(Tag))
+		{
+			continue;
+		}
+		if (!NameSubstring.IsEmpty() && !Test.TestName.Contains(NameSubstring))
+		{
+			continue;
+		}
+
+		++NumTest;
+		if (!RunParametricTest(Test, InRuntimeFilter))
+		{
+			++NumTestFailed;
+		}
+	}
+
+	if (NumTest == 0)
+	{
+		UE_LOG(LogNNE, Display, TEXT("No test selected to run (on %d tests registered)."), ParametricTests.TestSetups.Num());
+		return true;
+	}
+	else if (NumTestFailed == 0)
+	{
+		UE_LOG(LogNNE, Display, TEXT("SUCCEED! All %d tests selected passed (%d tests are registered)."), NumTest, ParametricTests.TestSetups.Num());
+		return true;
+	}
+	else
+	{
+		UE_LOG(LogNNE, Error, TEXT("FAILED! %d test(s) failed, on the %d test selected to run (%d tests are registered)."), NumTestFailed, NumTest, ParametricTests.TestSetups.Num());
+		return false;
+	}
+}
+	
+static FString FindArg(const FString& ArgName, const TArray< FString >& Args)
+{
+	FString Arg(TEXT(""));
+	int32 ArgNameIndex = Args.Find(ArgName);
+		
+	if (ArgNameIndex == INDEX_NONE || ArgNameIndex + 1 >= Args.Num())
+	{
+		return Arg;
+	}
+	else
+	{
+		return Args[ArgNameIndex + 1];
+	}
+}
+
+static FAutoConsoleCommand RunTestCommand(
+	TEXT("nne.test.run"), TEXT("Run all tests. Use -name to filter by name (substring). Use -tag to filter by tag. Use -runtime to only run for the provided runtime, default is to use filter set from setruntimefilter command."),
+	FConsoleCommandWithArgsDelegate::CreateStatic(
+		[](const TArray< FString >& Args)
+		{
+			FString Name = FindArg(TEXT("-name"), Args);
+			FString Tag = FindArg(TEXT("-tag"), Args);
+			FString Runtime = FindArg(TEXT("-runtime"), Args);
+			double StartTime = FPlatformTime::Seconds();
+			bool bTestSucceeded = RunParametricTests(Name, Tag, Runtime);
+			double EndTime = FPlatformTime::Seconds();
+			double TimeForTest = static_cast<float>(EndTime - StartTime);
+			if (bTestSucceeded)
+			{ 
+				UE_LOG(LogNNE, Display, TEXT("Tests succeeded in %0.1f seconds."), TimeForTest);
+			}
+			else
+			{
+				UE_LOG(LogNNE, Warning, TEXT("Tests FAILED in %0.1f seconds."), TimeForTest);
+			}
+		}
+	)
+);
+
+static FAutoConsoleCommand RunSmokeTestCommand(
+	TEXT("nne.test.smokerun"), TEXT("Run all smoke tests. Use -name to additionaly filter by name. Use -runtime to only run for the provided runtime, default is to run on all runtime but NNXRuntimeCPU (too slow at the moment, see comment in code)."),
+	FConsoleCommandWithArgsDelegate::CreateStatic(
+		[](const TArray< FString >& Args)
+		{
+			FString Name = FindArg(TEXT("-name"), Args);
+			FString Runtime = FindArg(TEXT("-runtime"), Args);
+			double StartTime = FPlatformTime::Seconds();
+			bool bTestSucceeded = RunParametricTests(Name, TEXT("smoketest"), Runtime);
+			double EndTime = FPlatformTime::Seconds();
+			double TimeForTest = static_cast<float>(EndTime - StartTime);
+			if (bTestSucceeded)
+			{
+				UE_LOG(LogNNE, Display, TEXT("Smoke tests succeeded in %0.1f seconds."), TimeForTest);
+			}
+			else
+			{
+				UE_LOG(LogNNE, Warning, TEXT("Smoke tests FAILED in %0.1f seconds."), TimeForTest);
+			}
+		}
+	)
+);
+
+#if WITH_DEV_AUTOMATION_TESTS
+	
+IMPLEMENT_SIMPLE_AUTOMATION_TEST_PRIVATE(FNNXParametricTestBase, FAutomationTestBase, "NNXParametricTest", EAutomationTestFlags::ApplicationContextMask | EAutomationTestFlags::FeatureMask | EAutomationTestFlags::EngineFilter, __FILE__, __LINE__)
+bool FNNXParametricTestBase::RunTest(const FString& Parameters) { return false; }
+
+class FNNXParametricTest : public FNNXParametricTestBase
+{
+	FTests::FTestSetup Test;
+
+public:
+	FNNXParametricTest(const FTests::FTestSetup& InTest) : FNNXParametricTestBase(InTest.TestName), Test(InTest) {}
+	virtual ~FNNXParametricTest() {}
+	virtual FString GetTestSourceFileName() const override { return "From Json"; }//Should return source json file path
+	virtual int32 GetTestSourceFileLine() const override { return 0; }
+protected:
+	virtual FString GetBeautifiedTestName() const override { return Test.TestName; }
+	bool RunTest(const FString& Parameter)
+	{
+		return RunParametricTest(Test, AutomationRuntimeFilter);
+	}
+};
+
+class FParametricTestAutomationRegistry
+{
+	
+	TArray<FNNXParametricTest*> RegisteredTests;
+
+public:
+	FParametricTestAutomationRegistry()
+	{
+		Refresh();
+	}
+		
+	~FParametricTestAutomationRegistry()
+	{
+		Clear();
+	}
+
+	void Clear()
+	{
+		for (FNNXParametricTest* Test : RegisteredTests)
+		{
+			delete Test;
+		}
+		RegisteredTests.Empty();
+	}
+
+	void Refresh()
+	{
+		Clear();
+
+		for (const FTests::FTestSetup& Test : ParametricTests.TestSetups)
+		{
+			RegisteredTests.Emplace(new FNNXParametricTest(Test));
+		}
+	}
+};
+
+static FParametricTestAutomationRegistry ParametricTestAutomationRegistry;
+
+#endif //WITH_DEV_AUTOMATION_TESTS
+
+static FAutoConsoleCommand TestReloadCommand(
+	TEXT("nne.test.reload"), TEXT("Reload NNX tests definition from Json."),
+	FConsoleCommandWithArgsDelegate::CreateStatic(
+		[](const TArray< FString >& Args)
+		{
+			ParametricTests.ReloadTestDescriptionsFromJson();
+				
+			#if WITH_DEV_AUTOMATION_TESTS
+				ParametricTestAutomationRegistry.Refresh();
+			#endif
+		}
+	)
+);
+
+bool InitializeParametricTests()
+{
+	bool bResult = ParametricTests.ReloadTestDescriptionsFromJson();
+#if WITH_DEV_AUTOMATION_TESTS
+	ParametricTestAutomationRegistry.Refresh();
+#endif
+	return bResult;
+}
+
+} // namespace UE::NNEQA::Private
diff --git a/Engine/Plugins/Experimental/NNE/Source/NNEQA/Private/NNEQAParametricTest.h b/Engine/Plugins/Experimental/NNE/Source/NNEQA/Private/NNEQAParametricTest.h
new file mode 100644
--- /dev/null
+++ b/Engine/Plugins/Experimental/NNE/Source/NNEQA/Private/NNEQAParametricTest.h
@@ -0,0 +1,8 @@
+// Copyright Epic Games, Inc. All Rights Reserved.
+
+#pragma once
+
+namespace UE::NNEQA::Private
+{
+	bool InitializeParametricTests();
+} // namespace UE::NNEQA::Internal
\ No newline at end of file
diff --git a/Engine/Plugins/Experimental/NNE/Source/NNEQA/Private/NNEQAUtils.cpp b/Engine/Plugins/Experimental/NNE/Source/NNEQA/Private/NNEQAUtils.cpp
new file mode 100644
--- /dev/null
+++ b/Engine/Plugins/Experimental/NNE/Source/NNEQA/Private/NNEQAUtils.cpp
@@ -0,0 +1,819 @@
+// Copyright Epic Games, Inc. All Rights Reserved.
+
+#include "NNEQAUtils.h"
+#include "HAL/UnrealMemory.h"
+#include "Kismet/GameplayStatics.h"
+#include "NNECore.h"
+#include "NNECoreAttributeMap.h"
+#include "NNECoreModelData.h"
+#include "NNECoreRuntime.h"
+#include "NNECoreRuntimeCPU.h"
+#include "NNXCore.h"
+#include "NNXInferenceModel.h"
+#include "NNXModelOptimizerInterface.h"
+#include "UObject/ObjectPtr.h"
+#include "UObject/WeakInterfacePtr.h"
+
+namespace UE::NNEQA::Private 
+{
+	using namespace NNX;
+	
+	static void FillTensors(
+		TConstArrayView<NNECore::Internal::FTensor> TensorsFromTestSetup,
+		TConstArrayView<NNECore::FTensorDesc> TensorDescsFromModel,
+		TArray<NNECore::Internal::FTensor>& OutTensors)
+	{
+		OutTensors = TensorsFromTestSetup;
+		if (OutTensors.IsEmpty())
+		{
+			for (const NNECore::FTensorDesc& SymbolicTensorDesc : TensorDescsFromModel)
+			{
+				NNECore::Internal::FTensor TensorDesc = NNECore::Internal::FTensor::MakeFromSymbolicDesc(SymbolicTensorDesc);
+				OutTensors.Emplace(TensorDesc);
+			}
+		}
+	}
+	
+	ElementWiseCosTensorInitializer::ElementWiseCosTensorInitializer(ENNETensorDataType InDataType, uint32 InTensorIndex)
+	: DataType(InDataType), TensorIndex(InTensorIndex)
+	{
+	}
+
+	float ElementWiseCosTensorInitializer::operator() (uint32 ElementIndex) const 
+	{
+		const constexpr uint32 IndexOffsetBetweenTensor = 9;
+		switch (DataType)
+		{
+
+		case ENNETensorDataType::Boolean:
+			return (ElementIndex + IndexOffsetBetweenTensor * TensorIndex) % 2;
+		case ENNETensorDataType::Char:
+		case ENNETensorDataType::Int8:
+		case ENNETensorDataType::Int16:
+		case ENNETensorDataType::Int32:
+		case ENNETensorDataType::Int64:
+			return 10.0f * FMath::Cos((float)(ElementIndex + IndexOffsetBetweenTensor * TensorIndex));
+		case ENNETensorDataType::UInt8:
+		case ENNETensorDataType::UInt16:
+		case ENNETensorDataType::UInt32:
+		case ENNETensorDataType::UInt64:
+			return 10.0f * FMath::Abs(FMath::Cos((float)(ElementIndex + IndexOffsetBetweenTensor * TensorIndex)));
+		default:
+			//case ENNETensorDataType::None:
+			//case ENNETensorDataType::Half:
+			//case ENNETensorDataType::Double:
+			//case ENNETensorDataType::Float:
+			//case ENNETensorDataType::Complex64:
+			//case ENNETensorDataType::Complex128:
+			//case ENNETensorDataType::BFloat16:
+			return FMath::Cos((float)(ElementIndex + IndexOffsetBetweenTensor * TensorIndex));
+		}
+	};
+
+	TArray<char> GenerateTensorDataForTest(const NNECore::Internal::FTensor& Tensor, std::function<float(uint32)> ElementInitializer)
+	{
+		const uint32 NumberOfElements = Tensor.GetVolume();
+		const int32 ElementByteSize = Tensor.GetElemByteSize();
+		const uint64 BufferSize = Tensor.GetDataSize();
+		const ENNETensorDataType DataType = Tensor.GetDataType();
+		TArray<char> Buffer;
+			
+		Buffer.SetNum(BufferSize);
+		char* DestinationPtr = Buffer.GetData();
+
+		for (uint32 i = 0; i != NumberOfElements; ++i)
+		{
+			const float FloatData = ElementInitializer(i);
+			const int32 IntData = (int32)FloatData;
+			const int64 Int64Data = (int64)FloatData;
+			const uint32 UIntData = (uint32)FloatData;
+			const char BoolData = FloatData != 0.0f ? 1 : 0;
+			switch (DataType)
+			{
+			case ENNETensorDataType::Float:
+				check(sizeof(FloatData) == ElementByteSize);
+				FMemory::Memcpy(DestinationPtr, &FloatData, ElementByteSize); break;
+			case ENNETensorDataType::Int32:
+				check(sizeof(IntData) == ElementByteSize);
+				FMemory::Memcpy(DestinationPtr, &IntData, ElementByteSize); break;
+			case ENNETensorDataType::Int64:
+				check(sizeof(Int64Data) == ElementByteSize);
+				FMemory::Memcpy(DestinationPtr, &Int64Data, ElementByteSize); break;
+			case ENNETensorDataType::UInt32:
+				check(sizeof(UIntData) == ElementByteSize);
+				FMemory::Memcpy(DestinationPtr, &UIntData, ElementByteSize); break;
+			case ENNETensorDataType::Boolean:
+				check(sizeof(BoolData) == ElementByteSize);
+				FMemory::Memcpy(DestinationPtr, &BoolData, ElementByteSize); break;
+			default:
+				FMemory::Memzero(DestinationPtr, ElementByteSize);
+			}
+			DestinationPtr += ElementByteSize;
+		}
+
+		return Buffer;
+	}
+
+	static void FillInputTensorBindingsCPU(TConstArrayView<NNECore::Internal::FTensor> Tensors, TConstArrayView<FTests::FTensorData> TensorsData,
+		TArray<TArray<char>>& OutMemBuffers, TArray<NNECore::FTensorBindingCPU>& OutBindings)
+	{
+		OutBindings.Empty();
+		OutMemBuffers.Empty();
+		check(TensorsData.Num() == 0 || Tensors.Num() == TensorsData.Num());
+
+		for (int Index = 0; Index < Tensors.Num(); ++Index)
+		{
+			const NNECore::Internal::FTensor& Tensor = Tensors[Index];
+			TArray<char>& MemBuffer = OutMemBuffers.Emplace_GetRef();
+			if (TensorsData.Num() == 0 || TensorsData[Index].Num() == 0)
+			{
+				ElementWiseCosTensorInitializer Initializer(Tensor.GetDataType(), Index);
+				MemBuffer = GenerateTensorDataForTest(Tensor, Initializer);
+			}
+			else
+			{
+				MemBuffer = TensorsData[Index];
+			}
+			check(MemBuffer.Num() == Tensor.GetDataSize());
+			
+			OutBindings.Emplace(NNECore::FTensorBindingCPU{(void*)MemBuffer.GetData(),(uint64)MemBuffer.Num()} );
+		}
+	}
+
+	static void FillOutputTensorBindingsCPU(TConstArrayView<NNECore::Internal::FTensor> Tensors, TArray<TArray<char>>& OutMemBuffers, TArray<NNECore::FTensorBindingCPU>& OutBindings)
+	{
+		OutBindings.Empty();
+		OutMemBuffers.Empty();
+		for (int Index = 0; Index < Tensors.Num(); ++Index)
+		{
+			TArray<char>& MemBuffer = OutMemBuffers.Emplace_GetRef();
+			MemBuffer.SetNumUninitialized(Tensors[Index].GetDataSize());
+			char constexpr MagicNumber = 0x5b;
+			FMemory::Memset(MemBuffer.GetData(), MagicNumber, MemBuffer.Num());
+			OutBindings.Emplace(NNECore::FTensorBindingCPU{ (void*)MemBuffer.GetData(),(uint64)MemBuffer.Num() });
+		}
+	}
+
+	static void FillInputTensorBindings(TConstArrayView<NNECore::Internal::FTensor> Tensors, TConstArrayView<FTests::FTensorData> TensorsData,
+		TArray<TArray<char>>& OutMemBuffers, TArray<FMLTensorBinding>& OutBindings)
+	{
+		OutBindings.Empty();
+		OutMemBuffers.Empty();
+		check(TensorsData.Num() == 0 || Tensors.Num() == TensorsData.Num());
+
+		for (int Index = 0; Index < Tensors.Num(); ++Index)
+		{
+			const NNECore::Internal::FTensor& Tensor = Tensors[Index];
+			TArray<char>& MemBuffer = OutMemBuffers.Emplace_GetRef();
+			if (TensorsData.Num() == 0 || TensorsData[Index].Num() == 0)
+			{
+				ElementWiseCosTensorInitializer Initializer(Tensor.GetDataType(), Index);
+				MemBuffer = GenerateTensorDataForTest(Tensor, Initializer);
+			}
+			else
+			{
+				MemBuffer = TensorsData[Index];
+			}
+			check(MemBuffer.Num() == Tensor.GetDataSize());
+			OutBindings.Emplace(FMLTensorBinding::FromCPU((void*)(MemBuffer.GetData()), MemBuffer.Num()));
+		}
+	}
+
+	static void FillOutputTensorBindings(TConstArrayView<NNECore::Internal::FTensor> Tensors, TArray<TArray<char>>& OutMemBuffers, TArray<FMLTensorBinding>& OutBindings)
+	{
+		OutBindings.Empty();
+		OutMemBuffers.Empty();
+		for (int Index = 0; Index < Tensors.Num(); ++Index)
+		{
+			TArray<char>& MemBuffer = OutMemBuffers.Emplace_GetRef();
+			MemBuffer.SetNumUninitialized(Tensors[Index].GetDataSize());
+			char constexpr MagicNumber = 0x5b;
+			FMemory::Memset(MemBuffer.GetData(), MagicNumber, MemBuffer.Num());
+			OutBindings.Emplace(FMLTensorBinding::FromCPU((void*)(MemBuffer.GetData()), MemBuffer.Num()));
+		}
+	}
+
+	template<typename T> FString ShapeToString(TConstArrayView<T> Shape)
+	{
+		FString TestSuffix(TEXT("["));
+		bool bIsFirstDim = true;
+		for (T size : Shape)
+		{
+			if (!bIsFirstDim) TestSuffix += TEXT(",");
+			TestSuffix += FString::Printf(TEXT("%d"), size);
+			bIsFirstDim = false;
+		}
+		TestSuffix += TEXT("]");
+		return TestSuffix;
+	}
+	template FString ShapeToString<int32>(TConstArrayView<int32> Shape);
+	template FString ShapeToString<uint32>(TConstArrayView<uint32> Shape);
+
+	FString TensorToString(const NNECore::Internal::FTensor& desc)
+	{
+		TArrayView<const uint32> Shape(desc.GetShape().GetData());
+
+		FString TensorDesc;
+		TensorDesc.Reserve(50);
+		TensorDesc += FString::Printf(TEXT("Name: %s, Shape: "), *desc.GetName());
+		TensorDesc += ShapeToString(Shape);
+		const FString& DataTypeName = StaticEnum<ENNETensorDataType>()->GetNameStringByValue(static_cast<int64>(desc.GetDataType()));
+		TensorDesc += FString::Printf(TEXT(" DataType: %s"), *DataTypeName);
+
+		return TensorDesc;
+	}
+
+	FString TensorToString(const NNECore::Internal::FTensor& TensorDesc, TConstArrayView<char> TensorData)
+	{
+		FString TensorLog;
+		TensorLog.Reserve(50);
+
+		TensorLog += TensorToString(TensorDesc);
+		TensorLog += TEXT(", Data: ");
+
+		const constexpr uint32 MAX_DATA_TO_LOG = 10;
+		const uint32 MaxIndex = FMath::Min(MAX_DATA_TO_LOG, TensorDesc.GetVolume());
+        const uint32 ElementByteSize = TensorDesc.GetElemByteSize();
+		for (uint32 i = 0; i < MaxIndex; ++i)
+		{
+			const uint32 ByteOffset = i * ElementByteSize;
+			const char* Data = TensorData.GetData() + ByteOffset;
+			check((int32)ByteOffset <= TensorData.Num());
+
+			switch (TensorDesc.GetDataType())
+			{
+				case ENNETensorDataType::Float:
+					TensorLog += FString::Printf(TEXT("%0.2f"), *(float*)Data); break;
+				case ENNETensorDataType::Int32:
+					TensorLog += FString::Printf(TEXT("%d"), *(int32*)Data); break;
+				case ENNETensorDataType::UInt32:
+					TensorLog += FString::Printf(TEXT("%u"), *(uint32*)Data); break;
+				case ENNETensorDataType::Boolean:
+					check(ElementByteSize == 1);
+					TensorLog += FString::Printf(TEXT("%s"), *(bool*)Data ? "true" : "false"); break;
+				default:
+					TensorLog += TEXT("?");
+			}
+			if (i < MaxIndex)
+				TensorLog += TEXT(", ");
+		}
+		if (MaxIndex < TensorDesc.GetVolume())
+			TensorLog += TEXT(",...");
+
+		return TensorLog;
+	}
+
+	template<typename T> bool CompareTensorData(
+		const NNECore::Internal::FTensor& RefTensorDesc,   const TArray<char>& RefRawBuffer,
+		const NNECore::Internal::FTensor& OtherTensorDesc, const TArray<char>& OtherRawBuffer,
+		float AbsoluteTolerance, float RelativeTolerance)
+	{
+
+		const T* RefBuffer = (T*)RefRawBuffer.GetData();
+		const T* OtherBuffer = (T*)OtherRawBuffer.GetData();
+		
+		const uint32 Volume = RefTensorDesc.GetVolume();
+		const uint32 ElementByteSize = RefTensorDesc.GetElemByteSize();
+
+		check(Volume == OtherTensorDesc.GetVolume());
+		check(Volume * ElementByteSize == RefRawBuffer.Num());
+		check(Volume * ElementByteSize == OtherRawBuffer.Num());
+
+		bool bTensorMemMatch = true;
+
+		float WorstAbsoluteError = 0.0f;
+		int32 WorstAbsoluteErrorIndex = -1;
+		float WorstAbsoluteErrorRef = 0.0f;
+		float WorstAbsoluteErrorOther = 0.0f;
+
+		float WorstRelativeError = 0.0f;
+		int32 WorstRelativeErrorIndex = -1;
+		float WorstRelativeErrorRef = 0.0f;
+		float WorstRelativeErrorOther = 0.0f;
+
+		int32 NumExtraNaNsInResults = 0;
+		int32 FirstExtraNaNIndex = -1;
+		int32 NumMissingNaNsInResults = 0;
+		int32 FirstMissingNaNIndex = -1;
+
+		for (uint32 i = 0; i < Volume; ++i)
+		{
+			//All type are converted to float for comparison purpose
+			const float Result = (float)OtherBuffer[i];
+			const float Reference = (float)RefBuffer[i];
+
+			
+			if (FMath::IsNaN(Result) && !FMath::IsNaN(Reference))
+			{
+				bTensorMemMatch = false;
+				++NumExtraNaNsInResults;
+				if (FirstExtraNaNIndex == -1)
+				{
+					FirstExtraNaNIndex = i;
+				}
+			}
+			if (!FMath::IsNaN(Result) && FMath::IsNaN(Reference))
+			{
+				bTensorMemMatch = false;
+				++NumMissingNaNsInResults;
+				if (FirstMissingNaNIndex == -1)
+				{
+					FirstMissingNaNIndex = i;
+				}
+			}
+			
+			if (FMath::IsNaN(Result) || FMath::IsNaN(Reference))
+			{
+				continue;
+			}
+			const float AbsoluteError = FMath::Abs<float>(Result - Reference);
+			const float AbsoluteRef = FMath::Abs<float>(Reference);
+			const float AbsoluteThreshold = AbsoluteTolerance + RelativeTolerance * AbsoluteRef;
+			/* from https://numpy.org/doc/stable/reference/generated/numpy.isclose.html */
+			if(AbsoluteError > AbsoluteThreshold)
+			{
+				bTensorMemMatch = false;
+				// Depending on dominating term, contribute to different statistics
+				if(AbsoluteError > RelativeTolerance * AbsoluteRef)
+				{
+					WorstRelativeError = AbsoluteError / AbsoluteRef;
+					WorstRelativeErrorIndex = i;
+					WorstRelativeErrorRef = Reference;
+					WorstRelativeErrorOther = Result;
+				}
+				else
+				{
+					WorstAbsoluteError = AbsoluteError;
+					WorstAbsoluteErrorIndex = i;
+					WorstAbsoluteErrorRef = Reference;
+					WorstAbsoluteErrorOther = Result;
+				}
+			}
+		}
+
+		if (!bTensorMemMatch)
+		{
+			FString AbsoluteErrorMessage = WorstAbsoluteErrorIndex == -1 ? FString() : FString::Printf(
+				TEXT("LogNNE: Worst absolute tolerance violation %.3e (tol %.3e) at position %d, got %.3e expected %.3e\n"),
+				WorstAbsoluteError, AbsoluteTolerance, WorstAbsoluteErrorIndex, WorstAbsoluteErrorOther, WorstAbsoluteErrorRef
+			);
+			FString RelativeErrorMessage = WorstRelativeErrorIndex == -1 ? FString() : FString::Printf(
+				TEXT("LogNNE: Worst relative tolerance violation %.3e (tol %.3e) at position %d, got %.3e expected %.3e\n"), 
+				WorstRelativeError, RelativeTolerance, WorstRelativeErrorIndex, WorstRelativeErrorOther, WorstRelativeErrorRef
+			); 
+			UE_LOG(LogNNE, Error, TEXT("Tensor data do not match.\n"
+				"%s"
+				"%s"
+				"LogNNE: Num unexpected NaNs %d (First at index %d), Num missing NaNs %d (first at index %d)"),
+				*AbsoluteErrorMessage,
+				*RelativeErrorMessage,
+				NumExtraNaNsInResults, FirstExtraNaNIndex, NumMissingNaNsInResults, FirstMissingNaNIndex);
+			UE_LOG(LogNNE, Error, TEXT("   Expected : %s"), *TensorToString(RefTensorDesc, RefRawBuffer));
+			UE_LOG(LogNNE, Error, TEXT("   But got  : %s"), *TensorToString(OtherTensorDesc, OtherRawBuffer));
+			return false;
+		}
+
+		return true;
+	}
+
+	bool VerifyTensorResult(
+		const NNECore::Internal::FTensor& RefTensor, const TArray<char>& RefRawBuffer,
+		const NNECore::Internal::FTensor& OtherTensor, const TArray<char>& OtherRawBuffer,
+		float AbsoluteTolerance, float RelativeTolerance)
+	{
+		if (RefTensor.GetShape() != OtherTensor.GetShape())
+		{
+			TArrayView<const uint32> RefShape(RefTensor.GetShape().GetData());
+			TArrayView<const uint32> OtherShape(OtherTensor.GetShape().GetData());
+			UE_LOG(LogNNE, Error, TEXT("Tensor shape do not match.\nExpected: %s\nGot:      %s"), *ShapeToString(RefShape), *ShapeToString(OtherShape));
+			return false;
+		}
+
+		if (RefTensor.GetDataType() == ENNETensorDataType::Float)
+		{
+			return CompareTensorData<float>(RefTensor, RefRawBuffer, OtherTensor, OtherRawBuffer, AbsoluteTolerance, RelativeTolerance);
+		}
+		else if (RefTensor.GetDataType() == ENNETensorDataType::Boolean)
+		{
+			check(RefTensor.GetElemByteSize() == 1);
+			return CompareTensorData<bool>(RefTensor, RefRawBuffer, OtherTensor, OtherRawBuffer, AbsoluteTolerance, RelativeTolerance);
+		}
+		else if (RefTensor.GetDataType() == ENNETensorDataType::Int32)
+		{
+			return CompareTensorData<int32>(RefTensor, RefRawBuffer, OtherTensor, OtherRawBuffer, AbsoluteTolerance, RelativeTolerance);
+		}
+		else if (RefTensor.GetDataType() == ENNETensorDataType::UInt32)
+		{
+			return CompareTensorData<uint32>(RefTensor, RefRawBuffer, OtherTensor, OtherRawBuffer, AbsoluteTolerance, RelativeTolerance);
+		}
+		else
+		{
+			UE_LOG(LogNNE, Error, TEXT("Tensor comparison for this type of tensor not implemented"));
+			return false;
+		}
+
+		return true;
+	}
+
+	TUniquePtr<NNECore::IModelCPU> GetInferenceModel(const FNNIModelRaw& ONNXModelData, const FString& RuntimeName)
+	{
+		TWeakInterfacePtr<INNERuntime>    RuntimeForModelData      = NNECore::GetRuntime<INNERuntime>(RuntimeName);
+		TWeakInterfacePtr<INNERuntimeCPU> RuntimeForInferenceModel = NNECore::GetRuntime<INNERuntimeCPU>(RuntimeName);
+		if (!RuntimeForModelData.IsValid() || !RuntimeForInferenceModel.IsValid())
+		{
+			UE_LOG(LogNNE, Error, TEXT("Can't get %s runtime."), *RuntimeName);
+			return TUniquePtr<NNECore::IModelCPU>();
+		}
+
+		TObjectPtr<UNNEModelData> ModelData = NewObject<UNNEModelData>();
+		ModelData->Init(FString("onnx"), ONNXModelData.Data);
+		TUniquePtr<NNECore::IModelCPU> InferenceModel = RuntimeForInferenceModel->CreateModelCPU(ModelData);
+		return InferenceModel;
+	}
+	
+	bool RunTestInferenceCPU(const FNNIModelRaw& ONNXModelData, const FTests::FTestSetup& TestSetup,
+		const FString& RuntimeName, TArray<NNECore::Internal::FTensor>& OutOutputTensors, TArray<TArray<char>>& OutOutputMemBuffers)
+	{
+		OutOutputMemBuffers.Empty();
+
+		TUniquePtr<NNECore::IModelCPU> InferenceModel = GetInferenceModel(ONNXModelData, RuntimeName);
+
+		if (!InferenceModel.IsValid())
+		{
+			UE_LOG(LogNNE, Error, TEXT("Could not create Inference model."));
+			return false;
+		}
+
+		// If test does not ask for specific input/output, fill with model default converting variable dim to 1 if any.
+		TArray<NNECore::Internal::FTensor> InputTensors;
+
+		FillTensors(TestSetup.Inputs, InferenceModel->GetInputTensorDescs(), InputTensors);
+		FillTensors(TestSetup.Outputs, InferenceModel->GetOutputTensorDescs(), OutOutputTensors);
+
+		// bind tensors to memory (CPU) and initialize
+		TArray<NNECore::FTensorBindingCPU> InputBindings;
+		TArray<NNECore::FTensorBindingCPU> OutputBindings;
+		TArray<TArray<char>> InputMemBuffers;
+
+		FillInputTensorBindingsCPU(InputTensors, TestSetup.InputsData, InputMemBuffers, InputBindings);
+		FillOutputTensorBindingsCPU(OutOutputTensors, OutOutputMemBuffers, OutputBindings);
+
+		// To help for debugging sessions.
+#ifdef UE_BUILD_DEBUG
+		const constexpr int32 NumTensorPtrForDebug = 3;
+		float* InputsAsFloat[NumTensorPtrForDebug];
+		float* OutputsAsFloat[NumTensorPtrForDebug];
+		FMemory::Memzero(InputsAsFloat, NumTensorPtrForDebug * sizeof(float*));
+		FMemory::Memzero(OutputsAsFloat, NumTensorPtrForDebug * sizeof(float*));
+		for (int32 i = 0; i < NumTensorPtrForDebug && i < InputMemBuffers.Num(); ++i)
+		{
+			InputsAsFloat[i] = (float*)InputMemBuffers[i].GetData();
+		}
+		for (int32 i = 0; i < NumTensorPtrForDebug && i < OutOutputMemBuffers.Num(); ++i)
+		{
+			OutputsAsFloat[i] = (float*)OutOutputMemBuffers[i].GetData();
+		}
+#endif //UE_BUILD_DEBUG
+
+		TArray<NNECore::FTensorShape> InputShapes;
+		for (const NNECore::Internal::FTensor& Tensor : InputTensors)
+		{
+			InputShapes.Emplace(Tensor.GetShape());
+		}
+
+		// Setup inputs
+		if (0 != InferenceModel->SetInputTensorShapes(InputShapes))
+		{
+			UE_LOG(LogNNE, Error, TEXT("Failed to set input tensor shapes."));
+			return false;
+		}
+
+		// Run inference
+		if (0 != InferenceModel->RunSync(InputBindings, OutputBindings))
+		{
+			UE_LOG(LogNNE, Error, TEXT("Failed to run the model."));
+			return false;
+		}
+
+		// Verify output shapes are as expected
+		TConstArrayView<NNECore::FTensorShape> OutputShapes = InferenceModel->GetOutputTensorShapes();
+		if (OutputShapes.Num() != OutOutputTensors.Num())
+		{
+			UE_LOG(LogNNE, Error, TEXT("Expected %d output tensors, got %d."), OutOutputTensors.Num(), OutputShapes.Num());
+			return false;
+		}
+		for (int i = 0; i < OutOutputTensors.Num(); ++i)
+		{
+			const NNECore::FTensorShape& RefTensorShape = OutOutputTensors[i].GetShape();
+			const NNECore::FTensorShape& OtherTensorShape = OutputShapes[i];
+			if (RefTensorShape != OtherTensorShape)
+			{
+				TArrayView<const uint32> RefShape(RefTensorShape.GetData());
+				TArrayView<const uint32> OtherShape(OtherTensorShape.GetData());
+				UE_LOG(LogNNE, Error, TEXT("Output shape do not match at index %d.\nExpected: %s\nGot:      %s"), i, *ShapeToString(RefShape), *ShapeToString(OtherShape));
+				return false;
+			}
+		}
+
+		return true;
+	}
+
+	static bool RunTestInference(const FNNIModelRaw& ONNXModelData, const FTests::FTestSetup& TestSetup, 
+		IRuntime* Runtime, TArray<NNECore::Internal::FTensor>& OutOutputTensors, TArray<TArray<char>>& OutOutputMemBuffers)
+	{
+		OutOutputMemBuffers.Empty();
+
+		TUniquePtr<FMLInferenceModel> InferenceModel = Runtime->CreateModel(Runtime->CreateModelData(FString("onnx"), ONNXModelData.Data));
+		
+		if (!InferenceModel.IsValid())
+		{
+			UE_LOG(LogNNE, Error, TEXT("Could not create Inference model."));
+			return false;
+		}
+
+		// If test does not ask for specific input/output, fill with model default converting variable dim to 1 if any.
+		TArray<NNECore::Internal::FTensor> InputTensors;
+		
+		FillTensors(TestSetup.Inputs, InferenceModel->GetInputTensorDescs(), InputTensors);
+		FillTensors(TestSetup.Outputs, InferenceModel->GetOutputTensorDescs(), OutOutputTensors);
+		
+		// bind tensors to memory (CPU) and initialize
+		TArray<FMLTensorBinding> InputBindings;
+		TArray<FMLTensorBinding> OutputBindings;
+		TArray<TArray<char>> InputMemBuffers;
+		
+		FillInputTensorBindings(InputTensors, TestSetup.InputsData, InputMemBuffers, InputBindings);
+		FillOutputTensorBindings(OutOutputTensors, OutOutputMemBuffers, OutputBindings);
+
+		// To help for debugging sessions.
+		#ifdef UE_BUILD_DEBUG
+		const constexpr int32 NumTensorPtrForDebug = 3;
+		float* InputsAsFloat[NumTensorPtrForDebug];
+		float* OutputsAsFloat[NumTensorPtrForDebug];
+		FMemory::Memzero(InputsAsFloat, NumTensorPtrForDebug * sizeof(float*));
+		FMemory::Memzero(OutputsAsFloat, NumTensorPtrForDebug * sizeof(float*));
+		for (int32 i = 0; i < NumTensorPtrForDebug && i < InputMemBuffers.Num(); ++i)
+		{
+			InputsAsFloat[i] = (float*)InputMemBuffers[i].GetData();
+		}
+		for (int32 i = 0; i < NumTensorPtrForDebug && i < OutOutputMemBuffers.Num(); ++i)
+		{
+			OutputsAsFloat[i] = (float*)OutOutputMemBuffers[i].GetData();
+		}
+		#endif //UE_BUILD_DEBUG
+
+		TArray<NNECore::FTensorShape> InputShapes;
+		for (const NNECore::Internal::FTensor& Tensor : InputTensors)
+		{
+			InputShapes.Emplace(Tensor.GetShape());
+		}
+
+		// Setup inputs
+		if (0 != InferenceModel->SetInputTensorShapes(InputShapes))
+		{
+			UE_LOG(LogNNE, Error, TEXT("Failed to set input tensor shapes."));
+			return false;
+		}
+
+		// Run inference
+		if (0 != InferenceModel->RunSync(InputBindings, OutputBindings))
+		{
+			UE_LOG(LogNNE, Error, TEXT("Failed to run the model."));
+			return false;
+		}
+
+		// Verify output shapes are as expected
+		TConstArrayView<NNECore::FTensorShape> OutputShapes = InferenceModel->GetOutputTensorShapes();
+		if (OutputShapes.Num() != OutOutputTensors.Num())
+		{
+			UE_LOG(LogNNE, Error, TEXT("Expected %d output tensors, got %d."), OutOutputTensors.Num(), OutputShapes.Num());
+			return false;
+		}
+		for (int i = 0; i < OutOutputTensors.Num(); ++i)
+		{
+			const NNECore::FTensorShape& RefTensorShape = OutOutputTensors[i].GetShape();
+			const NNECore::FTensorShape& OtherTensorShape = OutputShapes[i];
+			if (RefTensorShape != OtherTensorShape)
+			{
+				TArrayView<const uint32> RefShape(RefTensorShape.GetData());
+				TArrayView<const uint32> OtherShape(OtherTensorShape.GetData());
+				UE_LOG(LogNNE, Error, TEXT("Output shape do not match at index %d.\nExpected: %s\nGot:      %s"), i, *ShapeToString(RefShape), *ShapeToString(OtherShape));
+				return false;
+			}
+		}
+		
+		return true;
+	}
+
+	bool RunTestInferenceAndCompareToRef(const FTests::FTestSetup& TestSetup, IRuntime* Runtime, const FNNIModelRaw& ONNXModel, 
+		TArrayView<TArray<char>> RefOutputMemBuffers, TArrayView<NNECore::Internal::FTensor> RefOutputTensors)
+	{
+		TArray<TArray<char>> OutputMemBuffers;
+		TArray<NNECore::Internal::FTensor> OutputTensors;
+
+		const FString& RuntimeName = Runtime->GetRuntimeName();
+		const float AbsoluteTolerance = TestSetup.GetAbsoluteToleranceForRuntime(RuntimeName);
+		const float RelativeTolerance = TestSetup.GetRelativeToleranceForRuntime(RuntimeName);
+		bool bTestSuceeded = true;
+	
+		if (!RunTestInference(ONNXModel, TestSetup, Runtime, OutputTensors, OutputMemBuffers))
+		{
+			UE_LOG(LogNNE, Error, TEXT("Error running inference for engine %s."), *RuntimeName);
+			return false;
+		}
+
+		if (OutputTensors.Num() == RefOutputTensors.Num())
+		{
+			for (int i = 0; i < OutputTensors.Num(); ++i)
+			{
+				bTestSuceeded &= VerifyTensorResult(
+					RefOutputTensors[i], RefOutputMemBuffers[i],
+					OutputTensors[i], OutputMemBuffers[i],
+					AbsoluteTolerance, RelativeTolerance);
+			}
+		}
+		else
+		{
+			UE_LOG(LogNNE, Error, TEXT("Expecting %d output tensor(s), got %d."), RefOutputTensors.Num(), OutputTensors.Num());
+			bTestSuceeded = false;
+		}
+		return bTestSuceeded;
+	}
+
+	bool RunTestInferenceAndCompareToRef(const FTests::FTestSetup& TestSetup, const FString& RuntimeName, const FNNIModelRaw& ONNXModel,
+		TArrayView<TArray<char>> RefOutputMemBuffers, TArrayView<NNECore::Internal::FTensor> RefOutputTensors)
+	{
+		TArray<TArray<char>> OutputMemBuffers;
+		TArray<NNECore::Internal::FTensor> OutputTensors;
+
+		const float AbsoluteTolerance = TestSetup.GetAbsoluteToleranceForRuntime(RuntimeName);
+		const float RelativeTolerance = TestSetup.GetRelativeToleranceForRuntime(RuntimeName);
+		bool bTestSuceeded = true;
+
+		if (!RunTestInferenceCPU(ONNXModel, TestSetup, RuntimeName, OutputTensors, OutputMemBuffers))
+		{
+			UE_LOG(LogNNE, Error, TEXT("Error running inference for engine %s."), *RuntimeName);
+			return false;
+		}
+
+		if (OutputTensors.Num() == RefOutputTensors.Num())
+		{
+			for (int i = 0; i < OutputTensors.Num(); ++i)
+			{
+				bTestSuceeded &= VerifyTensorResult(
+					RefOutputTensors[i], RefOutputMemBuffers[i],
+					OutputTensors[i], OutputMemBuffers[i],
+					AbsoluteTolerance, RelativeTolerance);
+			}
+		}
+		else
+		{
+			UE_LOG(LogNNE, Error, TEXT("Expecting %d output tensor(s), got %d."), RefOutputTensors.Num(), OutputTensors.Num());
+			bTestSuceeded = false;
+		}
+		return bTestSuceeded;
+	}
+
+	bool CompareONNXModelInferenceAcrossRuntimes(const FNNIModelRaw& ONNXModel, const FNNIModelRaw& ONNXModelVariadic, const FTests::FTestSetup& TestSetup, const FString& RuntimeFilter)
+	{
+		FString CurrentPlatform = UGameplayStatics::GetPlatformName();
+		if (TestSetup.AutomationExcludedPlatform.Contains(CurrentPlatform))
+		{
+			UE_LOG(LogNNE, Display, TEXT("Skipping test of '%s' for platform %s (by config)"), *TestSetup.TargetName, *CurrentPlatform);
+			return true;
+		}
+		UE_LOG(LogNNE, Display, TEXT("Starting tests of '%s'"), *TestSetup.TargetName);
+
+		// Reference runtime
+		const FString& RefName = TEXT("NNERuntimeCPU");
+		const float AbsoluteTolerance = TestSetup.GetAbsoluteToleranceForRuntime(RefName);
+		const float RelativeTolerance = TestSetup.GetRelativeToleranceForRuntime(RefName);
+		TArray<TArray<char>> RefOutputMemBuffers;
+		TArray<NNECore::Internal::FTensor> RefOutputTensors;
+		bool bAllTestsSucceeded = true;
+
+		if (!RunTestInferenceCPU(ONNXModel, TestSetup, RefName, RefOutputTensors, RefOutputMemBuffers))
+		{
+			UE_LOG(LogNNE, Error, TEXT("Error running reference inference with engine %s."), *RefName);
+			return false;
+		}
+
+		// If outputs data have been defined by test setup we check that reference match it
+		for (int i = 0; i < TestSetup.OutputsData.Num(); ++i)
+		{
+			if (!TestSetup.OutputsData[i].IsEmpty())
+			{
+				bAllTestsSucceeded &= VerifyTensorResult(
+					TestSetup.Outputs[i], TestSetup.OutputsData[i],
+					RefOutputTensors[i], RefOutputMemBuffers[i],
+					AbsoluteTolerance, RelativeTolerance);
+			}
+		}
+		if (!bAllTestsSucceeded)
+		{
+			UE_LOG(LogNNE, Error, TEXT("Expecting output from test setup are no matched by reference runtime."));
+		}
+
+		// Test against other runtime
+		//NNX Registry
+		for (auto Runtime : GetAllRuntimes())
+		{
+			const FString& RuntimeName = Runtime->GetRuntimeName();
+			if (RuntimeName == TEXT("NNXRuntimeCPU"))
+			{
+				continue;
+			}
+			if (!RuntimeFilter.IsEmpty() && !RuntimeFilter.Contains(RuntimeName))
+			{
+				continue;
+			}
+			FString TestResult;
+			if (TestSetup.AutomationExcludedRuntime.Contains(RuntimeName) ||
+				TestSetup.AutomationExcludedPlatformRuntimeCombination.Contains(TPair<FString, FString>(CurrentPlatform, RuntimeName)))
+			{
+				TestResult = TEXT("skipped (by config)");
+			}
+			else
+			{
+				bool bShouldRunVariadicTest = (ONNXModelVariadic.Format != ENNXInferenceFormat::Invalid);
+				bShouldRunVariadicTest &= !(RuntimeName == "NNXRuntimeDML");
+
+				bool bTestSuceeded = RunTestInferenceAndCompareToRef(TestSetup, Runtime, ONNXModel, RefOutputMemBuffers, RefOutputTensors);
+
+				if (bShouldRunVariadicTest)
+				{
+					if (!bTestSuceeded)
+					{
+						UE_LOG(LogNNE, Error, TEXT("Failed running static test."));
+					}
+					bool bVariadicTestSuceeded = RunTestInferenceAndCompareToRef(TestSetup, Runtime, ONNXModelVariadic, RefOutputMemBuffers, RefOutputTensors);
+					if (!bVariadicTestSuceeded)
+					{
+						bTestSuceeded = false;
+						UE_LOG(LogNNE, Error, TEXT("Failed running variadic test."));
+					}
+				}
+				TestResult = bTestSuceeded ? TEXT("SUCCESS") : TEXT("FAILED");
+				bAllTestsSucceeded &= bTestSuceeded;
+			}
+
+			UE_LOG(LogNNE, Display, TEXT("  %s tests: %s"), *RuntimeName, *TestResult);
+		}
+
+		//NNE Registry
+		for (auto Runtime : NNECore::GetAllRuntimes())
+		{
+			const FString& RuntimeName = Runtime->GetRuntimeName();
+			if (RuntimeName == RefName)
+			{
+				continue;
+			}
+			if (!RuntimeFilter.IsEmpty() && !RuntimeFilter.Contains(RuntimeName))
+			{
+				continue;
+			}
+			FString TestResult;
+			if (TestSetup.AutomationExcludedRuntime.Contains(RuntimeName) ||
+				TestSetup.AutomationExcludedPlatformRuntimeCombination.Contains(TPair<FString, FString>(CurrentPlatform, RuntimeName)))
+			{
+				TestResult = TEXT("skipped (by config)");
+			}
+			else
+			{
+				bool bShouldRunVariadicTest = (ONNXModelVariadic.Format != ENNXInferenceFormat::Invalid);
+				bShouldRunVariadicTest &= !(RuntimeName == "NNEXRuntimeDML");
+
+				bool bTestSuceeded = RunTestInferenceAndCompareToRef(TestSetup, Runtime->GetRuntimeName(), ONNXModel, RefOutputMemBuffers, RefOutputTensors);
+
+				if (bShouldRunVariadicTest)
+				{
+					if (!bTestSuceeded)
+					{
+						UE_LOG(LogNNE, Error, TEXT("Failed running static test."));
+					}
+					bool bVariadicTestSuceeded = RunTestInferenceAndCompareToRef(TestSetup, Runtime->GetRuntimeName(), ONNXModelVariadic, RefOutputMemBuffers, RefOutputTensors);
+					if (!bVariadicTestSuceeded)
+					{
+						bTestSuceeded = false;
+						UE_LOG(LogNNE, Error, TEXT("Failed running variadic test."));
+					}
+				}
+				TestResult = bTestSuceeded ? TEXT("SUCCESS") : TEXT("FAILED");
+				bAllTestsSucceeded &= bTestSuceeded;
+			}
+
+			UE_LOG(LogNNE, Display, TEXT("  %s tests: %s"), *RuntimeName, *TestResult);
+		}
+		
+		return bAllTestsSucceeded;
+	}
+	FTests::FTestSetup& FTests::AddTest(const FString & Category, const FString & ModelOrOperatorName, const FString & TestSuffix)
+	{
+		FString TestName = Category + ModelOrOperatorName + TestSuffix;
+		//Test name should be unique
+		check(nullptr == TestSetups.FindByPredicate([TestName](const FTestSetup& Other) { return Other.TestName == TestName; }));
+		TestSetups.Emplace(Category, ModelOrOperatorName, TestSuffix);
+		return TestSetups.Last(0);
+	}
+
+} // namespace UE::NNEQA::Private
\ No newline at end of file
diff --git a/Engine/Plugins/Experimental/NNE/Source/NNEQA/Private/NNEQAUtils.h b/Engine/Plugins/Experimental/NNE/Source/NNEQA/Private/NNEQAUtils.h
new file mode 100644
--- /dev/null
+++ b/Engine/Plugins/Experimental/NNE/Source/NNEQA/Private/NNEQAUtils.h
@@ -0,0 +1,90 @@
+// Copyright Epic Games, Inc. All Rights Reserved.
+
+#pragma once
+
+#include "Containers/UnrealString.h"
+#include <functional>
+#include "NNECoreAttributeMap.h"
+#include "NNXRuntime.h"
+#include "NNXRuntimeFormat.h"
+#include "NNECoreTypes.h"
+#include "NNECoreTensor.h"
+
+
+namespace UE::NNEQA::Private
+{
+	struct FTests
+	{
+		typedef TArray<char> FTensorData;
+		
+		struct FTestSetup
+		{
+			/* defaults from https://numpy.org/doc/stable/reference/generated/numpy.isclose.html */
+			static constexpr float DefaultAbsoluteTolerance = 1e-8f;
+			static constexpr float DefaultRelativeTolerance = 1e-5f;
+
+			FTestSetup(const FString& InTestCategory, const FString& InModelOrOperatorName, const FString& InTestSuffix) :
+				TestName(InTestCategory + InModelOrOperatorName + InTestSuffix),
+				TargetName(InModelOrOperatorName),
+				AbsoluteTolerance(DefaultAbsoluteTolerance),
+				RelativeTolerance(DefaultRelativeTolerance),
+				IsModelTest(false)
+			{}
+			
+			const FString TestName;
+			const FString TargetName;
+			float AbsoluteTolerance;
+			float RelativeTolerance;
+			bool IsModelTest;
+			TMap<FString, float> AbsoluteToleranceForRuntime;
+			TMap<FString, float> RelativeToleranceForRuntime;
+			TArray<NNECore::Internal::FTensor> Inputs;
+			TArray<NNECore::Internal::FTensor> Weights;
+			TArray<NNECore::Internal::FTensor> Outputs;
+			TArray<FTensorData> InputsData;
+			TArray<FTensorData> WeightsData;
+			TArray<FTensorData> OutputsData;
+			NNECore::FAttributeMap AttributeMap;
+			TArray<FString> Tags;
+			TArray<FString> AutomationExcludedRuntime;
+			TArray<FString> AutomationExcludedPlatform;
+			TArray<TPair<FString, FString>> AutomationExcludedPlatformRuntimeCombination;
+
+			float GetAbsoluteToleranceForRuntime(const FString& RuntimeName) const
+			{
+				const float* SpecializedValue = AbsoluteToleranceForRuntime.Find(RuntimeName);
+				return (SpecializedValue != nullptr) ? *SpecializedValue : AbsoluteTolerance;
+
+			}
+			float GetRelativeToleranceForRuntime(const FString& RuntimeName) const
+			{
+				const float* SpecializedValue = RelativeToleranceForRuntime.Find(RuntimeName);
+				return (SpecializedValue != nullptr) ? *SpecializedValue : RelativeTolerance;
+			}
+		};
+
+
+		FTestSetup& AddTest(const FString& Category, const FString& ModelOrOperatorName, const FString& TestSuffix);
+
+		TArray<FTestSetup> TestSetups;
+	};
+
+	bool CompareONNXModelInferenceAcrossRuntimes(const FNNIModelRaw& ONNXModel, const FNNIModelRaw& ONNXModelVariadic, 
+		const FTests::FTestSetup& TestSetup, const FString& RuntimeFilter = TEXT(""));
+	
+	class ElementWiseCosTensorInitializer
+	{
+		ENNETensorDataType DataType;
+		uint32 TensorIndex;
+
+	public:
+		ElementWiseCosTensorInitializer(ENNETensorDataType InDataType, uint32 InTensorIndex);
+		float operator () (uint32 ElementIndex) const;
+	};
+
+	FString TensorToString(const NNECore::Internal::FTensor& Tensor);
+	FString TensorToString(const NNECore::Internal::FTensor& Tensor, TConstArrayView<char> TensorData);
+	template<typename T> FString ShapeToString(TConstArrayView<T> Shape);
+	TArray<char> GenerateTensorDataForTest(const NNECore::Internal::FTensor& Tensor, std::function<float(uint32)> ElementInitializer);
+
+} // namespace UE::NNEQA::Private
\ No newline at end of file
diff --git a/Engine/Plugins/Experimental/NNE/Source/NNEQA/Resources/NNEQATestDesc.json b/Engine/Plugins/Experimental/NNE/Source/NNEQA/Resources/NNEQATestDesc.json
new file mode 100644
--- /dev/null
+++ b/Engine/Plugins/Experimental/NNE/Source/NNEQA/Resources/NNEQATestDesc.json
@@ -0,0 +1,1046 @@
+{
+	"$schema": "./NNEQATestDescSchema.json",
+	"model_test_categories": [
+		{
+			"category": "OnnxZoo",
+			"is_model_test": true,
+			"targets": [
+				{
+					"target": "resnet50-v2-7",
+					"tags": [ "SmokeTest" ],
+					"runtimes": [
+						{ "name": "NNERuntimeORTDmlEditor" },
+						{ "name": "NNXRuntimeHlsl",  "skip": true, "comment": "Missing BatchNormalisation operator." },
+						{ "name": "NNXRuntimeDml",  "skip": true, "comment": "Model not supported yet." }
+					],
+					"absolute_tolerance": 1e-5
+				},
+				{
+					"target": "Mosaic-9",
+					"runtimes": [
+						{ "name": "NNERuntimeORTDmlEditor" },
+						{ "name": "NNERuntimeORTCpuEditor" },
+						{ "name": "NNXRuntimeHlsl" },
+						{ "name": "NNXRuntimeDml",  "skip": true, "comment": "Model not supported yet." }
+					],
+					"absolute_tolerance": 5.0e-2
+				}
+			]
+		},
+		{
+			"category": "InHouse",
+			"is_model_test": true,
+			"targets": [
+				{
+					"target": "NeuralMorphModel_global",
+					"tags": [ "SmokeTest" ],
+					"runtimes": [
+						{ "name": "NNERuntimeORTDmlEditor" },
+						{ "name": "NNXRuntimeHlsl" },
+						{ "name": "NNXRuntimeDml" }
+					],
+					"absolute_tolerance": 1e-5
+				},
+				{
+					"target": "NeuralMorphModel_global_Dax",
+					"tags": [ "SmokeTest" ],
+					"runtimes": [
+						{ "name": "NNERuntimeORTDmlEditor", "skip": true, "comment": "" },
+						{ "name": "NNXRuntimeHlsl",	"skip": true, "comment": "Model not supported yet."	},
+						{ "name": "NNXRuntimeDml",  "skip": true, "comment": "Model not supported yet." }
+					]
+				},
+				{
+					"target": "NeuralMorphModel_Local",
+					"tags": [ "SmokeTest" ],
+					"runtimes": [
+						{ "name": "NNERuntimeORTDmlEditor", "absolute_tolerance": 1e-6 },
+						{ "name": "NNXRuntimeHlsl",  "skip": true, "comment": "Missing Reshape operator." },
+						{ "name": "NNXRuntimeDml",  "skip": true, "comment": "Model not supported yet." }
+					]
+				},
+				{
+					"target": "VertexDeltaModel",
+					"runtimes": [
+						{ "name": "NNERuntimeORTDmlEditor", "absolute_tolerance": 1e-4 },
+						{ "name": "NNXRuntimeHlsl", "absolute_tolerance": 1e-4 },
+						{ "name": "NNXRuntimeDml",  "skip": true, "comment": "Model not supported yet." }
+					]
+				}
+			]
+		},
+		{
+			"category": "SimpleTests",
+			"is_model_test": true,
+			"targets": [
+				{
+					"target": "OneLayerWithWeights",
+					"runtimes": [
+						{ "name": "NNXRuntimeDml",  "skip": true, "comment": "Model not supported yet." }
+					]
+				},
+				{
+					"target": "TwoLayersNoWeight",
+					"tags": [ "SmokeTest" ],
+					"runtimes": [
+						{ "name": "NNXRuntimeDml",  "skip": true, "comment": "Model not supported yet." }
+					]
+				},
+				{
+					"target": "MultiOutput",
+					"tags": [ "SmokeTest" ],
+					"runtimes": [
+						{ "name": "NNXRuntimeDml",  "skip": true, "comment": "Model not supported yet." }
+					]
+				}
+			]
+		},
+		{
+			"category": "VariableInputShapes",
+			"is_model_test": true,
+			"targets": [
+				{
+					"target": "OneLayerNoWeightVariableInputShape",
+					"runtimes": [
+						{ "name": "NNXRuntimeDML",  "skip": true, "comment": "RDG does not support variable input sized tensor" }
+					]
+				},
+				{
+					"target": "TwoLayersNoWeightVariableInputShape",
+					"tags": [ "SmokeTest" ],
+					"runtimes": [
+						{ "name": "NNXRuntimeDML",  "skip": true, "comment": "RDG does not support variable input sized tensor" }
+					]
+				}
+			]
+		}
+	],
+
+	"operator_test_categories": [
+		{
+			"category": "Elementwise.Unary",
+			"targets": [
+				{	"target": "Abs", "additional_datasets": [ "Elementwise.Slow.OneInput" ]},
+				{	"target": "Acos", "tags": [ "SmokeTest" ], "runtimes": [
+						{ "name": "NNERuntimeORTDmlEditor", "absolute_tolerance":  1e-4},
+						{ "name": "NNXRuntimeDml",    "absolute_tolerance":  1e-4},
+						{ "name": "NNXRuntimeHlsl",   "absolute_tolerance":  2e-4}
+					]},
+				{	"target": "Acosh" },
+				{	"target": "Asin", "runtimes": [
+						{ "name": "NNERuntimeORTDmlEditor", "absolute_tolerance":  1e-4},
+						{ "name": "NNXRuntimeDml",    "absolute_tolerance":  1e-4},
+						{ "name": "NNXRuntimeHlsl",   "absolute_tolerance":  2e-4}
+					]},
+				{	"target": "Asinh", "runtimes": [
+						{ "name": "NNERuntimeORTDmlEditor", "absolute_tolerance":  1e-4},
+						{ "name": "NNXRuntimeDml",    "absolute_tolerance":  1e-4},
+						{ "name": "NNXRuntimeHlsl",   "absolute_tolerance":  2e-4}
+					]},
+				{	"target": "Atan", "runtimes": [
+						{ "name": "NNERuntimeORTDmlEditor", "absolute_tolerance":  1e-4},
+						{ "name": "NNXRuntimeDml",    "absolute_tolerance":  1e-4},
+						{ "name": "NNXRuntimeHlsl",   "absolute_tolerance":  1.5e-3}
+					]},
+				{	"target": "Atanh" },
+				{	"target": "BitShift", "skip": true, "comment": "TODO need attributes" },
+				{	"target": "Cast",     "skip": true, "comment": "TODO need attributes" },
+				{	"target": "Ceil"},
+				{	"target": "Clip", "additional_datasets": [ "Elementwise.ClipTestWithExtraInput" ], "runtimes": [
+						{ "name": "NNXRuntimeDml", "skip": true, "comment": "RDG need scalar tensor support"},
+						{ "name": "NNXRuntimeHlsl", "skip": true, "comment": "RDG need scalar tensor support"}						
+					]},
+				{	"target": "Cos" },
+				{	"target": "Cosh" },
+				{	"target": "Elu", "comment": "TODO test with not-default attribute", "runtimes": [
+						{ "name": "NNERuntimeORTDmlEditor", "absolute_tolerance":  1e-4},
+						{ "name": "NNXRuntimeDml",    "absolute_tolerance":  1e-4},
+						{ "name": "NNXRuntimeHlsl",   "absolute_tolerance":  1e-4}
+					]},
+				{	"target": "Erf", "runtimes": [
+						{ "name": "NNERuntimeORTDmlEditor", "absolute_tolerance":  1e-4},
+						{ "name": "NNXRuntimeDml",    "absolute_tolerance":  1e-4},
+						{ "name": "NNXRuntimeHlsl",   "absolute_tolerance":  5e-4, "comment": "NNXRuntimeHLSL Erf implementation is less precise than the DML one"}
+					]},
+				{	"target": "Exp" },
+				{	"target": "Floor" },
+				{	"target": "IsInf", "skip": true, "comment": "TODO add flexibility in test data setup + bool tensors" },
+				{	"target": "IsNan", "skip": true, "comment": "TODO add flexibility in test data setup + bool tensors" },
+				{	"target": "HardSigmoid",         "comment": "TODO test with not-default attribute" },
+				{	"target": "HardSwish", "runtimes": [ { "name": "NNXRuntimeDml", "skip": true, "comment": "Not implemented atm on NNXRuntimeDML."} ]},
+				{	"target": "LeakyRelu", "additional_datasets": [ "Elementwise.LeakyRelu.AttributeTest" ] },
+				{	"target": "Log", "runtimes": [
+						{ "name": "NNERuntimeORTDmlEditor"},
+						{ "name": "NNXRuntimeDml"},
+						{ "name": "NNXRuntimeHlsl"}
+					], "absolute_tolerance": 1e-6 },
+				{	"target": "Neg" },
+				{	"target": "Not", "input_type": "boolean", "output_type": "boolean", "runtimes": [
+						{ "name": "NNXRuntimeDml", "skip": true, "comment": "RDG need bool tensor support"},
+						{ "name": "NNXRuntimeHlsl", "skip": true, "comment": "RDG need bool tensor support"}						
+					]},
+				{	"target": "Reciprocal"},
+				{	"target": "Relu"},
+				{	"target": "Round"},
+				{	"target": "Selu", "comment": "TODO test with not-default attribute", "runtimes": [
+						{ "name": "NNERuntimeORTDmlEditor", "absolute_tolerance":  1e-4},
+						{ "name": "NNXRuntimeDml",    "absolute_tolerance":  1e-4},
+						{ "name": "NNXRuntimeHlsl",   "absolute_tolerance":  1e-4}
+					]},
+				{	"target": "Sigmoid"},
+				{	"target": "Sign" },
+				{	"target": "Sin", "runtimes": [
+						{ "name": "NNERuntimeORTDmlEditor", "absolute_tolerance":  1e-4},
+						{ "name": "NNXRuntimeDml",    "absolute_tolerance":  1e-4},
+						{ "name": "NNXRuntimeHlsl",   "absolute_tolerance":  1e-4}
+					]},
+				{	"target": "Sinh" },
+				{	"target": "Softplus"},
+				{	"target": "Softsign"},
+				{	"target": "Sqrt"},
+				{	"target": "Tan", "runtimes": [
+						{ "name": "NNERuntimeORTDmlEditor", "absolute_tolerance":  1e-4},
+						{ "name": "NNXRuntimeDml",    "absolute_tolerance":  1e-4},
+						{ "name": "NNXRuntimeHlsl",   "absolute_tolerance":  1e-4}
+					]},
+				{	"target": "Tanh", "runtimes": [
+						{ "name": "NNERuntimeORTDmlEditor", "absolute_tolerance":  1e-4},
+						{ "name": "NNXRuntimeDml",    "absolute_tolerance":  1e-4},
+						{ "name": "NNXRuntimeHlsl",   "absolute_tolerance":  1e-4}
+					]}
+			]
+		},
+		{
+			"category": "Elementwise.Binary",
+			"targets": [
+				{	"target": "Add", "additional_datasets": [ "Elementwise.Slow.TwoInputs" ]},
+				{	"target": "And", "input_type": "boolean", "output_type": "boolean", "runtimes": [
+						{ "name": "NNXRuntimeDml", "skip": true, "comment": "RDG need bool tensor support"},
+						{ "name": "NNXRuntimeHlsl", "skip": true, "comment": "RDG need bool tensor support"}						
+					]},
+				{	"target": "Div", "tags": [ "SmokeTest" ]},
+				{	"target": "Equal", "output_type": "boolean", "runtimes": [
+						{ "name": "NNXRuntimeDml", "skip": true, "comment": "RDG need bool tensor support"},
+						{ "name": "NNXRuntimeHlsl", "skip": true, "comment": "RDG need bool tensor support"}						
+					]},
+				{	"target": "Greater", "output_type": "boolean", "runtimes": [
+						{ "name": "NNXRuntimeDml", "skip": true, "comment": "RDG need bool tensor support"},
+						{ "name": "NNXRuntimeHlsl", "skip": true, "comment": "RDG need bool tensor support"}						
+					]},
+				{	"target": "GreaterOrEqual", "output_type": "boolean", "runtimes": [
+						{ "name": "NNXRuntimeDml", "skip": true, "comment": "RDG need bool tensor support"},
+						{ "name": "NNXRuntimeHlsl", "skip": true, "comment": "RDG need bool tensor support"}						
+					]},
+				{	"target": "Less", "output_type": "boolean", "runtimes": [
+						{ "name": "NNXRuntimeDml", "skip": true, "comment": "RDG need bool tensor support"},
+						{ "name": "NNXRuntimeHlsl", "skip": true, "comment": "RDG need bool tensor support"}						
+					]},
+				{	"target": "LessOrEqual", "output_type": "boolean", "runtimes": [
+						{ "name": "NNXRuntimeDml", "skip": true, "comment": "RDG need bool tensor support"},
+						{ "name": "NNXRuntimeHlsl", "skip": true, "comment": "RDG need bool tensor support"}						
+					]},
+				{	"target": "Mod", "skip": true, "comment": "TODO need to set fmod attribute to 1 for floating point tensors"},
+				{	"target": "Mul"},
+				{	"target": "Or", "input_type": "boolean", "output_type": "boolean", "runtimes": [
+						{ "name": "NNXRuntimeDml", "skip": true, "comment": "RDG need bool tensor support"},
+						{ "name": "NNXRuntimeHlsl", "skip": true, "comment": "RDG need bool tensor support"}						
+					]},
+				{	"target": "PRelu", "removed_datasets": [ 
+						"Elementwise.Binary.BidirectionalBroadcast", 
+						"Elementwise.Binary.LHSToRHSBroadcast" ], 
+					"comment": "Prelu does not support bidirectional broadcast. It only support slope(RHS) to input0(LHS)"},
+				{	"target": "Pow", "absolute_tolerance": 1e-4},
+				{	"target": "Sub"},
+				{	"target": "Xor", "input_type": "boolean", "output_type": "boolean", "runtimes": [
+						{ "name": "NNXRuntimeDml", "skip": true, "comment": "RDG need bool tensor support"},
+						{ "name": "NNXRuntimeHlsl", "skip": true, "comment": "RDG need bool tensor support"}						
+					]}
+			]
+		},
+		{
+			"category": "Elementwise.Variadic",
+			"additional_datasets": [ "Elementwise.Binary" ],
+			"runtimes": [ { "name": "NNXRuntimeDml", "skip": true} ],
+			"targets": [
+				{	"target": "Max", "additional_datasets": [ "Elementwise.Slow.TwoInputs" ]},
+				{	"target": "Mean", "tags": [ "SmokeTest" ]},
+				{	"target": "Min"},
+				{	"target": "Sum"}
+			]
+		},
+		{
+			"category": "Gemm",
+			"targets": [
+				{ "target": "Gemm", "tags": [ "SmokeTest" ] }
+			]
+		},
+		{
+			"category": "Convolution",
+			"runtimes": [
+				{ "name": "NNXRuntimeDml", "skip": true	},
+				{ "name": "NNERuntimeORTDmlEditor" },
+				{ "name": "NNXRuntimeHlsl" }
+			],
+			"targets": [
+				{ "target": "Conv", "absolute_tolerance": 1e-6 }
+			]
+		},
+		{
+			"category": "ConvTranspose",
+			"runtimes": [
+				{ "name": "NNXRuntimeDml", "skip": true	},
+				{ "name": "NNERuntimeORTDmlEditor" },
+				{ "name": "NNXRuntimeHlsl" }
+			],
+			"targets": [
+				{ "target": "ConvTranspose", "absolute_tolerance": 1e-6 }
+			]
+		},
+		{
+			"category": "MatMul",
+			"runtimes": [ { "name": "NNXRuntimeDml", "skip": true } ],
+			"targets": [
+				{ "target": "MatMul" }
+			]
+		},
+		{
+			"category": "PadOp",
+			"runtimes": [
+				{ "name": "NNXRuntimeDml", "skip": true	}
+			],
+			"targets": [
+				{ "target": "Pad" }
+			]
+		},
+		{
+			"category": "InstanceNormalization",
+			"runtimes": [
+				{ "name": "NNXRuntimeDml", "skip": true }
+			],
+			"targets": [
+				{ "target": "InstanceNormalization", "absolute_tolerance": 1e-5 }
+			]
+		},
+		{
+			"category": "Upsample",
+			"runtimes": [
+				{ "name": "NNXRuntimeDml", "skip": true	}
+			],
+			"targets": [
+				{ "target": "Upsample" }
+			]
+		}
+	],
+
+	"input_output_sets": [
+		{
+			"name": "VariableInputShapes.Regular",
+			"datasets": [
+				{ "inputs": [ { "shape": [ 1,1 ] } ], "comment": "if not specified output is defaulted as the first input" },
+				{ "inputs": [ { "shape": [ 2,3 ] } ] }
+			]
+		},
+		{
+			"name": "Elementwise.Unary.Regular",
+			"datasets": [
+				{ "inputs": [ { "shape": [ 1 ] } ] },
+				{ "inputs": [ { "shape": [ 1,2,3,4 ] } ] },
+				{ "inputs": [ { "shape": [ 1,512 ] } ] }
+			]
+		},
+		{
+			"name": "Elementwise.Unary.ZeroSized",
+			"datasets": [
+				{ "inputs": [ { "shape": [ 1,0,3 ] } ],
+					"runtimes": [
+						{ "name": "NNXRuntimeHLSL", "skip": true, "comment": "RDG does not support 0 sized tensor" },
+						{ "name": "NNXRuntimeDML",  "skip": true, "comment": "RDG does not support 0 sized tensor" }
+					]}
+			]
+		},
+		{
+			"name": "Elementwise.ClipTestWithExtraInput",
+			"datasets": [
+				{ "inputs": [ { "shape": [1], "comment": "input" }, { "shape": [1], "comment": "min" } ] },
+				{ "inputs": [ { "shape": [1], "comment": "input" }, { "shape": [1], "comment": "min" }, { "shape": [1], "comment": "max" } ] }
+			]
+		},
+		{
+			"name": "Elementwise.Slow.OneInput",
+			"datasets": [
+				{
+					"inputs": [ { "shape": [ 65536, 513 ] } ]
+				}
+			]
+		},
+		{
+			"name": "Elementwise.LeakyRelu.AttributeTest",
+			"datasets": [
+				{
+					"inputs": [ { "shape": [ 4 ] } ]
+				}
+			]
+		},
+		{
+			"name": "Elementwise.Binary.NoBroadcast",
+			"datasets": [
+				{ "inputs": [ { "shape": [ 1 ] }, { "shape": [ 1 ] } ]},
+				{ "inputs": [ { "shape": [ 4 ] }, { "shape": [ 4 ] } ]},
+				{ "inputs": [ { "shape": [ 2,3,4,5,6 ] }, { "shape": [ 2,3,4,5,6 ] } ] }
+			]
+		},
+		{
+			"name": "Elementwise.Variadic.NoBroadcast",
+			"datasets": [
+				{ "inputs": [ { "shape": [ 1 ] }, { "shape": [ 1 ] }, { "shape": [ 1 ] } ]},
+				{ "inputs": [ { "shape": [ 1 ] }, { "shape": [ 1 ] }, { "shape": [ 1 ] }, { "shape": [ 1 ] } ]},
+				{ "inputs": [ { "shape": [ 4 ] }, { "shape": [ 4 ] }, { "shape": [ 4 ] } ]},
+				{ "inputs": [ { "shape": [ 5 ] }, { "shape": [ 5 ] }, { "shape": [ 5 ] }, { "shape": [ 5 ] }, { "shape": [ 5 ] }, { "shape": [ 5 ] }, { "shape": [ 5 ] }, { "shape": [ 5 ] } ]}
+			]
+		},
+		{
+			"name": "Elementwise.Binary.ZeroSized",
+			"datasets": [
+				{ "inputs": [ { "shape": [ 1,0 ] }, { "shape": [ 1,1 ] } ],
+					"runtimes": [
+						{ "name": "NNXRuntimeHLSL", "skip": true, "comment": "RDG does not support 0 sized tensor" },
+						{ "name": "NNXRuntimeDML",  "skip": true, "comment": "RDG does not support 0 sized tensor" }
+					]}
+			]
+		},
+		{
+			"name": "Elementwise.Binary.RHSToLHSBroadcast",
+			"datasets": [
+				{ "inputs": [ { "shape": [ 2,3,4 ] }, { "shape": [ 1,1,1 ] } ]},
+				{ "inputs": [ { "shape": [ 2,3,4 ] }, { "shape": [ 2,1,1 ] } ]},
+				{ "inputs": [ { "shape": [ 2,3,4 ] }, { "shape": [ 1,1,4 ] } ]},
+				{ "inputs": [ { "shape": [ 2,3,4 ] }, { "shape": [ 1,3,1 ] } ]},
+				{ "inputs": [ { "shape": [ 2,3,4 ] }, { "shape": [ 1,3,4 ] } ]},
+				{ "inputs": [ { "shape": [ 2,3,4 ] }, { "shape": [ 2,1,4 ] } ]},
+				{ "inputs": [ { "shape": [ 2,3,4 ] }, { "shape": [ 1 ]     } ]},
+				{ "inputs": [ { "shape": [ 2,3,4 ] }, { "shape": [ 4 ]     } ]},
+				{ "inputs": [ { "shape": [ 2,3,4 ] }, { "shape": [ 1,1 ]   } ]},
+				{ "inputs": [ { "shape": [ 2,3,4 ] }, { "shape": [ 3,1 ]   } ]},
+				{ "inputs": [ { "shape": [ 2,3,4 ] }, { "shape": [ 1,4 ]   } ]},
+				{ "inputs": [ { "shape": [ 2,3,4 ] }, { "shape": [ 3,4 ]   } ]}
+			]
+		},
+		{
+			"name": "Elementwise.Binary.BidirectionalBroadcast",
+			"datasets": [
+				{ "inputs": [ { "shape": [ 2,3,1 ] }, { "shape": [ 1,1,4 ] } ], "outputs": [ { "shape": [ 2,3,4 ] } ]},
+				{ "inputs": [ { "shape": [ 2,1,4 ] }, { "shape": [ 2,3,1 ] } ], "outputs": [ { "shape": [ 2,3,4 ] } ]},
+				{ "inputs": [ { "shape": [ 2,3,1 ] }, { "shape": [ 4 ]     } ], "outputs": [ { "shape": [ 2,3,4 ] } ]},
+				{ "inputs": [ { "shape": [ 2,3,1 ] }, { "shape": [ 1,4 ]   } ], "outputs": [ { "shape": [ 2,3,4 ] } ]}
+			]
+		},
+		{
+			"name": "Elementwise.Variadic.MultiDirectionalBroadcast",
+			"datasets": [
+				{ "inputs": [ { "shape": [ 2,3,1 ] }, { "shape": [ 1,1,4 ] }, { "shape": [ 1,1,1 ] } ], "outputs": [ { "shape": [ 2,3,4 ] } ]},
+				{ "inputs": [ { "shape": [ 2,1,4 ] }, { "shape": [ 1,3,4 ] }, { "shape": [ 2,3,1 ] } ], "outputs": [ { "shape": [ 2,3,4 ] } ]},
+				{ "inputs": [ { "shape": [ 2,3,1 ] }, { "shape": [ 4 ]     }, { "shape": [ 3,1 ]   } ], "outputs": [ { "shape": [ 2,3,4 ] } ]},
+				{ "inputs": [ { "shape": [ 2,3,1 ] }, { "shape": [ 3,1 ]   }, { "shape": [ 4 ]     } ], "outputs": [ { "shape": [ 2,3,4 ] } ]}
+			]
+		},
+		{
+			"name": "Elementwise.Binary.LHSToRHSBroadcast",
+			"datasets": [
+				{ "inputs": [ { "shape": [ 1,1,1 ] }, { "shape": [ 2,3,4 ] } ], "outputs": [ { "shape": [ 2,3,4 ] } ]},
+				{ "inputs": [ { "shape": [ 2,1,1 ] }, { "shape": [ 2,3,4 ] } ], "outputs": [ { "shape": [ 2,3,4 ] } ]},
+				{ "inputs": [ { "shape": [ 1,1,4 ] }, { "shape": [ 2,3,4 ] } ], "outputs": [ { "shape": [ 2,3,4 ] } ]},
+				{ "inputs": [ { "shape": [ 1,3,1 ] }, { "shape": [ 2,3,4 ] } ], "outputs": [ { "shape": [ 2,3,4 ] } ]},
+				{ "inputs": [ { "shape": [ 1,3,4 ] }, { "shape": [ 2,3,4 ] } ], "outputs": [ { "shape": [ 2,3,4 ] } ]},
+				{ "inputs": [ { "shape": [ 2,1,4 ] }, { "shape": [ 2,3,4 ] } ], "outputs": [ { "shape": [ 2,3,4 ] } ]},
+				{ "inputs": [ { "shape": [ 1 ]     }, { "shape": [ 2,3,4 ] } ], "outputs": [ { "shape": [ 2,3,4 ] } ]},
+				{ "inputs": [ { "shape": [ 4 ]     }, { "shape": [ 2,3,4 ] } ], "outputs": [ { "shape": [ 2,3,4 ] } ]},
+				{ "inputs": [ { "shape": [ 1,1 ]   }, { "shape": [ 2,3,4 ] } ], "outputs": [ { "shape": [ 2,3,4 ] } ]},
+				{ "inputs": [ { "shape": [ 3,1 ]   }, { "shape": [ 2,3,4 ] } ], "outputs": [ { "shape": [ 2,3,4 ] } ]},
+				{ "inputs": [ { "shape": [ 1,4 ]   }, { "shape": [ 2,3,4 ] } ], "outputs": [ { "shape": [ 2,3,4 ] } ]},
+				{ "inputs": [ { "shape": [ 3,4 ]   }, { "shape": [ 2,3,4 ] } ], "outputs": [ { "shape": [ 2,3,4 ] } ]}
+			]
+		},
+		{
+			"name": "Elementwise.Slow.TwoInputs",
+			"datasets": [
+				{
+					"inputs": [ { "shape": [ 65536, 513 ] }, { "shape": [ 1 ] } ]
+				}
+			]
+		},
+		{
+			"name": "Gemm.TwoInputs.NoTranspose",
+			"datasets": [
+				{ "inputs": [ { "shape": [ 2, 2 ] }, { "shape": [ 2, 2 ] } ], "outputs": [ { "shape": [ 2, 2 ] } ]},
+				{ "inputs": [ { "shape": [ 3, 4 ] }, { "shape": [ 4, 5 ] } ], "outputs": [ { "shape": [ 3, 5 ] } ]}
+			],
+			"runtimes": [
+				{ "name": "NNERuntimeORTDmlEditor", "skip": true, "comment": "Variable dimension without bias is failing for NNERuntimeORTDmlEditor" }
+			]
+		},
+		{
+			"name": "Gemm.TwoInputs.TransposedA",
+			"datasets": [
+				{ "inputs": [ { "shape": [ 2, 2 ] }, { "shape": [ 2, 2 ] } ], "outputs": [ { "shape": [ 2, 2 ] } ]},
+				{ "inputs": [ { "shape": [ 4, 3 ] }, { "shape": [ 4, 5 ] } ], "outputs": [ { "shape": [ 3, 5 ] } ]}
+			],
+			"runtimes": [
+				{ "name": "NNERuntimeORTDmlEditor", "skip": true, "comment": "Variable dimension without bias is failing for NNERuntimeORTDmlEditor" }
+			]
+		},
+		{
+			"name": "Gemm.TwoInputs.TransposedB",
+			"datasets": [
+				{ "inputs": [ { "shape": [ 2, 2 ] }, { "shape": [ 2, 2 ] } ], "outputs": [ { "shape": [ 2, 2 ] } ]},
+				{ "inputs": [ { "shape": [ 3, 4 ] }, { "shape": [ 5, 4 ] } ], "outputs": [ { "shape": [ 3, 5 ] } ]}
+			],
+			"runtimes": [
+				{ "name": "NNERuntimeORTDmlEditor", "skip": true, "comment": "Variable dimension without bias is failing for NNERuntimeORTDmlEditor" }
+			]
+		},
+		{
+			"name": "Gemm.TwoInputs.AllTransposed",
+			"datasets": [
+				{ "inputs": [ { "shape": [ 2, 2 ] }, { "shape": [ 2, 2 ] } ], "outputs": [ { "shape": [ 2, 2 ] } ]},
+				{ "inputs": [ { "shape": [ 4, 3 ] }, { "shape": [ 5, 4 ] } ], "outputs": [ { "shape": [ 3, 5 ] } ]}
+			],
+			"runtimes": [
+				{ "name": "NNERuntimeORTDmlEditor", "skip": true, "comment": "Variable dimension without bias is failing for NNERuntimeORTDmlEditor" }
+			]
+		},
+		{
+			"name": "Gemm.ThreeInputs.NoTranspose",
+			"datasets": [
+				{ "inputs": [ { "shape": [ 3, 4 ] }, { "shape": [ 4, 5 ] }, { "shape": [ 3, 5 ] } ], "outputs": [ { "shape": [ 3, 5 ] } ]},
+				{ "inputs": [ { "shape": [ 1, 1 ] }, { "shape": [ 1, 5 ] }, { "shape": [ 5 ] } ], "outputs": [ { "shape": [ 1, 5 ] } ]}
+			]
+		},
+		{
+			"name": "Gemm.ThreeInputs.TransposedA",
+			"datasets": [
+				{ "inputs": [ { "shape": [ 4, 3 ] }, { "shape": [ 4, 5 ] }, { "shape": [ 3, 5 ] } ], "outputs": [ { "shape": [ 3, 5 ] } ]},
+				{ "inputs": [ { "shape": [ 1, 1 ] }, { "shape": [ 1, 5 ] }, { "shape": [ 5 ] } ], "outputs": [ { "shape": [ 1, 5 ] } ]}
+			]
+		},
+		{
+			"name": "Gemm.ThreeInputs.TransposedB",
+			"datasets": [
+				{ "inputs": [ { "shape": [ 3, 4 ] }, { "shape": [ 5, 4 ] }, { "shape": [ 3, 5 ] } ], "outputs": [ { "shape": [ 3, 5 ] } ]},
+				{ "inputs": [ { "shape": [ 1, 1 ] }, { "shape": [ 5, 1 ] }, { "shape": [ 5 ] } ], "outputs": [ { "shape": [ 1, 5 ] } ]}
+			]
+		},
+		{
+			"name": "Gemm.ThreeInputs.AllTransposed",
+			"datasets": [
+				{ "inputs": [ { "shape": [ 4, 3 ] }, { "shape": [ 5, 4 ] }, { "shape": [ 3, 5 ] } ], "outputs": [ { "shape": [ 3, 5 ] } ]},
+				{ "inputs": [ { "shape": [ 1, 1 ] }, { "shape": [ 5, 1 ] }, { "shape": [ 5 ] } ], "outputs": [ { "shape": [ 1, 5 ] } ]}
+			]
+		},
+		{ 
+			"name": "Convolution.DefaultPadding",
+			"datasets": [
+				{ "inputs": [ { "shape": [ 1, 1, 5, 5 ] }, { "shape": [ 1, 1, 3, 3 ] } ], "outputs": [ { "shape": [ 1, 1, 5, 5 ] } ] },
+				{ "inputs": [ { "shape": [ 1, 3, 5, 5 ] }, { "shape": [ 1, 3, 3, 3 ] } ], "outputs": [ { "shape": [ 1, 1, 5, 5 ] } ], "comment": "With channels" },
+				{ "inputs": [ { "shape": [ 1, 1, 5, 5 ] }, { "shape": [ 3, 1, 3, 3 ] } ], "outputs": [ { "shape": [ 1, 3, 5, 5 ] } ], "comment": "With out-channels" },
+				{ "inputs": [ { "shape": [ 1, 1, 5, 5 ] }, { "shape": [ 3, 1, 3, 3 ] }, { "shape": [ 3 ] } ], "outputs": [ { "shape": [ 1, 3, 5, 5 ] } ], "comment": "With bias" },
+				{ "inputs": [ { "shape": [ 3, 1, 5, 5 ] }, { "shape": [ 1, 1, 3, 3 ] } ], "outputs": [ { "shape": [ 3, 1, 5, 5 ] } ], "comment": "Batches" }
+			]
+		},
+		{ "name": "Convolution.WithoutOrAutoPadding", 			"datasets": [ { "inputs": [ { "shape": [ 1, 1, 5, 5 ] }, { "shape": [ 1, 1, 3, 3 ] } ], "outputs": [ { "shape": [ 1, 1, 3, 3 ] } ] } ] },
+		{ "name": "Convolution.WithStridesAndPadding", 			"datasets": [ { "inputs": [ { "shape": [ 1, 1, 7, 5 ] }, { "shape": [ 1, 1, 3, 3 ] } ], "outputs": [ { "shape": [ 1, 1, 4, 3 ] } ] } ] },
+		{ "name": "Convolution.WithoutPaddingAndStrides", 		"datasets": [ { "inputs": [ { "shape": [ 1, 1, 7, 5 ] }, { "shape": [ 1, 1, 3, 3 ] } ], "outputs": [ { "shape": [ 1, 1, 3, 2 ] } ] } ] },
+		{ "name": "Convolution.AsymmetricPaddingAndStrides", 	"datasets": [ { "inputs": [ { "shape": [ 1, 1, 7, 5 ] }, { "shape": [ 1, 1, 3, 3 ] } ], "outputs": [ { "shape": [ 1, 1, 4, 2 ] } ] } ] },
+		{ "name": "Convolution.WithGroups", 					"datasets": [ { "inputs": [ { "shape": [ 1, 3, 5, 5 ] }, { "shape": [ 3, 1, 3, 3 ] } ], "outputs": [ { "shape": [ 1, 3, 5, 5 ] } ] } ] },
+		{ "name": "Convolution.Large1", 						"datasets": [ { "inputs": [ { "shape": [ 1, 1, 17, 17 ] }, { "shape": [ 1, 1, 3, 3 ] } ], "outputs": [ { "shape": [ 1, 1, 17, 17 ] } ] } ] },
+		{ "name": "Convolution.Large2", 						"datasets": [ { "inputs": [ { "shape": [ 1, 1, 37, 37 ] }, { "shape": [ 1, 1, 3, 3 ] } ], "outputs": [ { "shape": [ 1, 1, 17, 17 ] } ] } ] },
+		{ "name": "ConvTranspose.Default", 		"datasets": [ { "inputs": [ { "shape": [ 1, 1, 3, 3 ] }, { "shape": [ 1, 2, 3, 3 ] } ], "outputs": [ { "shape": [ 1, 2, 5, 5 ] } ] } ] },
+		{ "name": "ConvTranspose.1D", 			"datasets": [ { "inputs": [ { "shape": [ 1, 1, 3 ] }, { "shape": [ 1, 2, 3 ] } ], "outputs": [ { "shape": [ 1, 2, 5 ] } ] } ] },
+		{ "name": "ConvTranspose.3D", 			"datasets": [ { "inputs": [ { "shape": [ 1, 1, 3, 4, 5 ] }, { "shape": [ 1, 2, 3, 3, 3 ] } ], "outputs": [ { "shape": [ 1, 2, 5, 6, 7 ] } ] } ] },
+		{ "name": "ConvTranspose.Attributes", 	"datasets": [ { "inputs": [ { "shape": [ 1, 1, 3, 3 ] }, { "shape": [ 1, 2, 3, 3 ] } ], "outputs": [ { "shape": [ 1, 2, 10, 8 ] } ] } ] },
+		{
+			"name": "ConvTranspose.Autopad",
+			"datasets": [ { "inputs": [ { "shape": [ 1, 1, 3, 3 ] }, { "shape": [ 1, 2, 3, 3 ] } ], "outputs": [ { "shape": [ 1, 2, 6, 6 ] } ] } ],
+			"runtimes": [
+				{ "name": "NNERuntimeORTDmlEditor", "skip": true, "comment": "TODO INVESTIGATE Returns 'The parameter is incorrect'" },
+				{ "name": "NNXRuntimeHlsl", "skip": true, "comment": "Fails until ORT 1.13, because they implemented 'auto_pad' differently (they flipped 'SAME_LOWER' and 'SAME_UPPER')" }
+			]
+		},
+		{
+			"name": "MatMul.Regular",
+			"datasets": [
+				{ "inputs": [ { "shape": [ 2, 2 ] }, { "shape": [ 2, 2 ] } ], "outputs": [ { "shape": [ 2, 2 ] } ]},
+				{ "inputs": [ { "shape": [ 3, 4 ] }, { "shape": [ 4, 5 ] } ], "outputs": [ { "shape": [ 3, 5 ] } ]},
+				{ "inputs": [ { "shape": [ 2, 3, 4 ] }, { "shape": [ 2, 4, 5 ] } ], "outputs": [ { "shape": [ 2, 3, 5 ] } ]},
+				{ "inputs": [ { "shape": [ 5, 4 ] }, { "shape": [ 2, 4, 3 ] } ], "outputs": [ { "shape": [ 2, 5, 3 ] } ]},
+				{ "inputs": [ { "shape": [ 2, 5, 4 ] }, { "shape": [ 4, 3 ] } ], "outputs": [ { "shape": [ 2, 5, 3 ] } ]},
+				{ "inputs": [ { "shape": [ 1, 4 ] }, { "shape": [ 4, 3 ] } ], "outputs": [ { "shape": [ 1, 3 ] } ]},
+				{ "inputs": [ { "shape": [ 1, 4 ] }, { "shape": [ 2, 4, 3 ] } ], "outputs": [ { "shape": [ 2, 1, 3 ] } ]},
+				{ "inputs": [ { "shape": [ 3, 4 ] }, { "shape": [ 4, 1 ] } ], "outputs": [ { "shape": [ 3, 1 ] } ]},
+				{ "inputs": [ { "shape": [ 2, 2, 3, 4 ] }, { "shape": [ 4, 1 ] } ], "outputs": [ { "shape": [ 2, 2, 3, 1 ] } ]},
+				{ "inputs": [ { "shape": [ 1, 6 ] }, { "shape": [ 6, 1 ] } ], "outputs": [ { "shape": [ 1, 1 ] } ]}
+			]
+		},
+		{
+			"name": "MatMul.Rank5",
+			"runtimes": [	
+				{ "name": "NNERuntimeORTDmlEditor", "skip": true, "comment": "DML metacommand only support 2 to 4 for inputs rank" },
+				{ "name": "NNXRuntimeHlsl"}
+			],
+			"datasets": [
+				{ "inputs": [ { "shape": [ 2, 2, 3, 5, 4 ] }, { "shape": [ 2, 1, 3, 4, 3 ] } ], "outputs": [ { "shape": [ 2, 2, 3, 5, 3 ] } ]},
+				{ "inputs": [ { "shape": [ 2, 1, 2, 5, 4 ] }, { "shape": [ 1, 3, 2, 4, 3 ] } ], "outputs": [ { "shape": [ 2, 3, 2, 5, 3 ] } ]},
+				{ "inputs": [ { "shape": [ 1, 1, 2, 5, 4 ] }, { "shape": [ 2, 3, 2, 4, 3 ] } ], "outputs": [ { "shape": [ 2, 3, 2, 5, 3 ] } ]},
+				{ "inputs": [ { "shape": [ 1, 2, 5, 4 ] }, { "shape": [ 2, 3, 2, 4, 3 ] } ], "outputs": [ { "shape": [ 2, 3, 2, 5, 3 ] } ]},
+				{ "inputs": [ { "shape": [ 5, 4 ] }, { "shape": [ 2, 1, 3, 4, 3 ] } ], "outputs": [ { "shape": [ 2, 1, 3, 5, 3 ] } ]},
+				{ "inputs": [ { "shape": [ 2, 1, 3, 5, 4 ] }, { "shape": [ 4, 3 ] } ], "outputs": [ { "shape": [ 2, 1, 3, 5, 3 ] } ]}
+			]
+		},
+		{
+			"name": "PadOp.Pad00Rank1",
+			"datasets": [
+				{ "inputs":  [ { "shape": [ 3 ] } ], "outputs": [ { "shape": [ 3 ] } ] }
+			]
+		},
+		{
+			"name": "PadOp.Pad03Rank1",
+			"datasets": [
+				{ "inputs":  [ { "shape": [ 4 ] } ], "outputs": [ { "shape": [ 7 ] } ] }
+			]
+		},
+		{
+			"name": "PadOp.Pad20Rank1",
+			"datasets": [
+				{ "inputs":  [ { "shape": [ 4 ] } ], "outputs": [ { "shape": [ 6 ] } ] }
+			]
+		},
+		{
+			"name": "PadOp.Pad41Rank1",
+			"datasets": [
+				{ "inputs":  [ { "shape": [ 5 ] } ], "outputs": [ { "shape": [ 10 ] } ] }
+			]
+		},
+		{
+			"name": "PadOp.Pad00_10_12_34_Rank4",
+			"datasets": [
+				{ "inputs":  [ { "shape": [ 1,32,224,128 ] } ], "outputs": [ { "shape": [ 1,33,227,135 ] } ] }
+			]
+		},
+		{
+			"name": "Upsample.ScaleFromInput",
+			"runtimes": [	
+				{ "name": "NNXRuntimeDml", "skip": true, "comment": "Scale should be a constant tensor for RDG" },
+				{ "name": "NNXRuntimeHlsl", "skip": true, "comment": "Scale should be a constant tensor for RDG" }
+			],
+			"datasets": [
+				{ "inputs":  [ { "shape": [ 1 ] },                            { "shape": [ 1 ], "source": [ "3"    ],                  "type": "float" } ], "outputs": [ { "shape": [ 3 ]} ] },
+				{ "inputs":  [ { "shape": [ 2 ] },                            { "shape": [ 1 ], "source": [ "2"    ],                  "type": "float" } ], "outputs": [ { "shape": [ 4 ]} ] },
+				{ "inputs":  [ { "shape": [ 3 ] },                            { "shape": [ 1 ], "source": [ "1"    ],                  "type": "float" } ], "outputs": [ { "shape": [ 3 ]} ] },
+				{ "inputs":  [ { "shape": [ 1, 128, 50, 55 ] },               { "shape": [ 4 ], "source": [ "1", "1", "2", "2" ],      "type": "float" } ], "outputs": [ { "shape": [ 1, 128, 100, 110 ]} ] },
+				{ "inputs":  [ { "shape": [ 2, 1, 3, 1, 2 ] },                { "shape": [ 5 ], "source": [ "1", "2", "2", "1", "2" ], "type": "float" } ], "outputs": [ { "shape": [ 2, 2, 6, 1, 4 ]} ] }
+			]
+		},
+		{
+			"name": "Upsample.ScaleFromWeight",
+			"datasets": [
+				{ "inputs":  [ { "shape": [ 1 ] } ],              "weights": [ { "shape": [ 1 ], "source": [ "3"    ],                  "type": "float" } ], "outputs": [ { "shape": [ 3 ]} ] },
+				{ "inputs":  [ { "shape": [ 2 ] } ],              "weights": [ { "shape": [ 1 ], "source": [ "2"    ],                  "type": "float" } ], "outputs": [ { "shape": [ 4 ]} ] },
+				{ "inputs":  [ { "shape": [ 3 ] } ],              "weights": [ { "shape": [ 1 ], "source": [ "1"    ],                  "type": "float" } ], "outputs": [ { "shape": [ 3 ]} ] },
+				{ "inputs":  [ { "shape": [ 1, 128, 50, 55 ] } ], "weights": [ { "shape": [ 4 ], "source": [ "1", "1", "2", "2" ],      "type": "float" } ], "outputs": [ { "shape": [ 1, 128, 100, 110 ]} ] },
+				{ "inputs":  [ { "shape": [ 2, 1, 3, 1, 2 ] } ],  "weights": [ { "shape": [ 5 ], "source": [ "1", "2", "2", "1", "2" ], "type": "float" } ], "outputs": [ { "shape": [ 2, 2, 6, 1, 4 ]} ] }
+			]
+		},
+		{
+			"name": "InstanceNormalization",
+			"datasets": [
+				{ "inputs": [ { "shape": [ 1, 1, 3 ] }, { "shape": [ 1 ] }, { "shape": [ 1 ] } ], "outputs": [ { "shape": [ 1, 1, 3 ] } ], "runtimes": [ { "name": "NNERuntimeORTDmlEditor", "skip": true, "comment": "Not supported" } ] },
+				{ "inputs": [ { "shape": [ 1, 2, 3 ] }, { "shape": [ 2 ] }, { "shape": [ 2 ] } ], "outputs": [ { "shape": [ 1, 2, 3 ] } ], "runtimes": [ { "name": "NNERuntimeORTDmlEditor", "skip": true, "comment": "Not supported" } ] },
+				{ "inputs": [ { "shape": [ 2, 2, 3 ] }, { "shape": [ 2 ] }, { "shape": [ 2 ] } ], "outputs": [ { "shape": [ 2, 2, 3 ] } ], "runtimes": [ { "name": "NNERuntimeORTDmlEditor", "skip": true, "comment": "Not supported" } ]},
+				{ "inputs": [ { "shape": [ 1, 2, 1, 3 ] }, { "shape": [ 2 ] }, { "shape": [ 2 ] } ], "outputs": [ { "shape": [ 1, 2, 1, 3 ] } ] },
+				{ "inputs": [ { "shape": [ 2, 2, 1, 3 ] }, { "shape": [ 2 ] }, { "shape": [ 2 ] } ], "outputs": [ { "shape": [ 2, 2, 1, 3 ] } ] },
+				{ "inputs": [ { "shape": [ 2, 2, 2, 2 ] }, { "shape": [ 2 ] }, { "shape": [ 2 ] } ], "outputs": [ { "shape": [ 2, 2, 2, 2 ] } ] },
+				{
+					"inputs": [ { "shape": [ 2, 2, 4, 1, 3 ] }, { "shape": [ 2 ] }, { "shape": [ 2 ] } ], "outputs": [ { "shape": [ 2, 2, 4, 1, 3 ] } ],
+					"runtimes": [
+						{ "name": "NNXRuntimeHLSL" },
+						{ "name": "NNERuntimeORTCpuEditor" },
+						{ "name": "NNERuntimeORTDmlEditor", "skip": true, "comment": "Not supported" }
+					]
+				},
+				{
+					"inputs": [ { "shape": [ 2, 2, 4, 4 ] }, { "shape": [ 2 ] }, { "shape": [ 2 ] } ], "outputs": [ { "shape": [ 2, 2, 4, 4 ] } ],
+					"runtimes": [
+						{ "name": "NNXRuntimeHLSL" },
+						{ "name": "NNERuntimeORTCpuEditor" },
+						{ "name": "NNERuntimeORTDmlEditor" }
+					]
+				},
+				{
+					"inputs": [ { "shape": [ 2, 2, 8, 8 ] }, { "shape": [ 2 ] }, { "shape": [ 2 ] } ], "outputs": [ { "shape": [ 2, 2, 8, 8 ] } ],
+					"runtimes": [
+						{ "name": "NNXRuntimeHLSL" },
+						{ "name": "NNERuntimeORTCpuEditor" },
+						{ "name": "NNERuntimeORTDmlEditor" }
+					]
+				},
+				{
+					"inputs": [ { "shape": [ 2, 128, 8, 8 ] }, { "shape": [ 128 ] }, { "shape": [ 128 ] } ], "outputs": [ { "shape": [ 2, 128, 8, 8 ] } ],
+					"runtimes": [
+						{ "name": "NNXRuntimeHLSL" },
+						{ "name": "NNERuntimeORTCpuEditor" },
+						{ "name": "NNERuntimeORTDmlEditor" }
+					]
+				},
+				{
+					"inputs": [ { "shape": [ 128, 2, 8, 8 ] }, { "shape": [ 2 ] }, { "shape": [ 2 ] } ], "outputs": [ { "shape": [ 128, 2, 8, 8 ] } ],
+					"runtimes": [
+						{ "name": "NNXRuntimeHLSL" },
+						{ "name": "NNERuntimeORTCpuEditor" },
+						{ "name": "NNERuntimeORTDmlEditor" }
+					]
+				},
+				{
+					"inputs": [ { "shape": [ 64, 64, 8, 8 ] }, { "shape": [ 64 ] }, { "shape": [ 64 ] } ], "outputs": [ { "shape": [ 64, 64, 8, 8 ] } ],
+					"runtimes": [
+						{ "name": "NNXRuntimeHLSL" },
+						{ "name": "NNERuntimeORTCpuEditor" },
+						{ "name": "NNERuntimeORTDmlEditor" }
+					]
+				},
+				{
+					"inputs": [ { "shape": [ 2, 2, 45, 45 ] }, { "shape": [ 2 ] }, { "shape": [ 2 ] } ], "outputs": [ { "shape": [ 2, 2, 45, 45 ] } ],
+					"runtimes": [
+						{ "name": "NNXRuntimeHLSL" },
+						{ "name": "NNERuntimeORTCpuEditor" },
+						{ "name": "NNERuntimeORTDmlEditor" }
+					]
+				},
+				{
+					"inputs": [ { "shape": [ 1, 128, 56, 56 ] }, { "shape": [ 128 ] }, { "shape": [ 128 ] } ], "outputs": [ { "shape": [ 1, 128, 56, 56 ] } ],
+					"runtimes": [
+						{ "name": "NNXRuntimeHLSL" },
+						{ "name": "NNERuntimeORTCpuEditor" },
+						{ "name": "NNERuntimeORTDmlEditor" }
+					],
+					"comment": "from 'mosaic-9-sim.onnx'"
+				}
+			]
+		}
+	],
+
+	"attribute_sets": [
+		{
+			"name": "Elementwise.LeakyRelu.AttributeTest",
+			"attribute_maps": [
+				{ "attributes": [ { "name": "alpha", "value": 0.1 } ] },
+				{ "attributes": [ { "name": "alpha", "value": 0.5 } ] },
+				{ "attributes": [ { "name": "alpha", "value": 1.0 } ] }
+			]
+		},
+		{
+			"name": "Gemm.TwoInputs.NoTranspose.Alpha",
+			"attribute_maps": [
+				{ "attributes": [ { "name": "alpha", "value": 0.0 } ] },
+				{ "attributes": [ { "name": "alpha", "value": 0.5 } ] },
+				{ "attributes": [ { "name": "alpha", "value": 1.0 } ] },
+				{ "attributes": [ { "name": "alpha", "value": 5.0 } ] }
+			]
+		},
+		{
+			"name": "Gemm.ThreeInputs.NoTranspose.Beta",
+			"multiply_with_attribute_sets": [ "Gemm.TwoInputs.NoTranspose.Alpha" ],
+			"attribute_maps": [
+				{ "attributes": [ { "name": "beta", "value": 0.0 } ] },
+				{ "attributes": [ { "name": "beta", "value": 0.5 } ] },
+				{ "attributes": [ { "name": "beta", "value": 1.0 } ] },
+				{ "attributes": [ { "name": "beta", "value": 5.0 } ] }
+			]
+		},
+		{
+			"name": "Gemm.TwoInputs.NoTranspose",
+			"attribute_maps": [
+				{ "attributes": [ { "name": "transA", "value": 0, "type": "Int32" }, { "name": "transB", "value": 0, "type": "Int32" } ] }
+			]
+		},
+		{
+			"name": "Gemm.ThreeInputs.NoTranspose",
+			"attribute_maps": [
+				{ "attributes": [ { "name": "transA", "value": 0, "type": "Int32" }, { "name": "transB", "value": 0, "type": "Int32" } ] }
+			]
+		},
+		{
+			"name": "Gemm.TwoInputs.TransposedA",
+			"attribute_maps": [
+				{ "attributes": [ { "name": "transA", "value": 1, "type": "Int32" } ] }
+			]
+		},
+		{
+			"name": "Gemm.TwoInputs.TransposedB",
+			"attribute_maps": [
+				{ "attributes": [ { "name": "transB", "value": 1, "type": "Int32" } ] }
+			]
+		},
+		{
+			"name": "Gemm.TwoInputs.AllTransposed",
+			"attribute_maps": [
+				{ "attributes": [ { "name": "transA", "value": 1, "type": "Int32" }, { "name": "transB", "value": 1, "type": "Int32" } ] }
+			]
+		},
+				{
+			"name": "Gemm.ThreeInputs.TransposedA",
+			"attribute_maps": [
+				{ "attributes": [ { "name": "transA", "value": 1, "type": "Int32" } ] }
+			]
+		},
+		{
+			"name": "Gemm.ThreeInputs.TransposedB",
+			"attribute_maps": [
+				{ "attributes": [ { "name": "transB", "value": 1, "type": "Int32" } ] }
+			]
+		},
+		{
+			"name": "Gemm.ThreeInputs.AllTransposed",
+			"attribute_maps": [
+				{ "attributes": [ { "name": "transA", "value": 1, "type": "Int32" }, { "name": "transB", "value": 1, "type": "Int32" } ] }
+			]
+		},
+		{
+			"name": "Convolution.DefaultPadding",
+			"attribute_maps": [
+				{ "attributes": [
+					{ "name": "auto_pad", "value": "NOTSET", "type": "String" },
+					{ "name": "dilations", "value": [ 1, 1 ], "type": "Int32Array" },
+					{ "name": "group", "value": 1, "type": "Int32" },
+					{ "name": "pads", "value": [ 1, 1, 1, 1 ], "type": "Int32Array" },
+					{ "name": "strides", "value": [ 1, 1 ], "type": "Int32Array"}
+				] }
+			]
+		},
+		{
+			"name": "Convolution.WithoutOrAutoPadding.WithoutPadding",
+			"attribute_maps": [
+				{ "attributes": [
+					{ "name": "auto_pad", "value": "NOTSET", "type": "String" },
+					{ "name": "dilations", "value": [ 1, 1 ], "type": "Int32Array" },
+					{ "name": "group", "value": 1, "type": "Int32" },
+					{ "name": "pads", "value": [ 0, 0, 0, 0 ], "type": "Int32Array" },
+					{ "name": "strides", "value": [ 1, 1 ], "type": "Int32Array"}
+				] }
+			]
+		},
+		{
+			"name": "Convolution.WithoutOrAutoPadding.AutoPaddingAndStrides",
+			"attribute_maps": [
+				{ "attributes": [
+					{ "name": "auto_pad", "value": "SAME_LOWER", "type": "String" },
+					{ "name": "dilations", "value": [ 1, 1 ], "type": "Int32Array" },
+					{ "name": "group", "value": 1, "type": "Int32" },
+					{ "name": "strides", "value": [ 2, 2 ], "type": "Int32Array"}
+				] }
+			]
+		},
+		{
+			"name": "Convolution.WithStridesAndPadding",
+			"attribute_maps": [
+				{ "attributes": [
+					{ "name": "auto_pad", "value": "NOTSET", "type": "String" },
+					{ "name": "dilations", "value": [ 1, 1 ], "type": "Int32Array" },
+					{ "name": "group", "value": 1, "type": "Int32" },
+					{ "name": "pads", "value": [ 1, 1, 1, 1 ], "type": "Int32Array" },
+					{ "name": "strides", "value": [ 2, 2 ], "type": "Int32Array"}
+				] }
+			]
+		},
+		{
+			"name": "Convolution.WithoutPaddingAndStrides",
+			"attribute_maps": [
+				{ "attributes": [
+					{ "name": "auto_pad", "value": "NOTSET", "type": "String" },
+					{ "name": "dilations", "value": [ 1, 1 ], "type": "Int32Array" },
+					{ "name": "group", "value": 1, "type": "Int32" },
+					{ "name": "pads", "value": [ 0, 0, 0, 0 ], "type": "Int32Array" },
+					{ "name": "strides", "value": [ 2, 2 ], "type": "Int32Array"}
+				] }
+			]
+		},
+		{
+			"name": "Convolution.AsymmetricPaddingAndStrides",
+			"attribute_maps": [
+				{ "attributes": [
+					{ "name": "auto_pad", "value": "NOTSET", "type": "String" },
+					{ "name": "dilations", "value": [ 1, 1 ], "type": "Int32Array" },
+					{ "name": "group", "value": 1, "type": "Int32" },
+					{ "name": "pads", "value": [ 1, 0, 1, 0 ], "type": "Int32Array" },
+					{ "name": "strides", "value": [ 2, 2 ], "type": "Int32Array"}
+				] }
+			]
+		},
+		{
+			"name": "Convolution.WithGroups",
+			"attribute_maps": [
+				{ "attributes": [
+					{ "name": "auto_pad", "value": "NOTSET", "type": "String" },
+					{ "name": "dilations", "value": [ 1, 1 ], "type": "Int32Array" },
+					{ "name": "group", "value": 3, "type": "Int32" },
+					{ "name": "pads", "value": [ 1, 1, 1, 1 ], "type": "Int32Array" },
+					{ "name": "strides", "value": [ 1, 1 ], "type": "Int32Array"}
+				] }
+			]
+		},
+		{
+			"name": "Convolution.Large1",
+			"attribute_maps": [
+				{ "attributes": [
+					{ "name": "auto_pad", "value": "SAME_UPPER", "type": "String" },
+					{ "name": "dilations", "value": [ 1, 1 ], "type": "Int32Array" },
+					{ "name": "group", "value": 1, "type": "Int32" },
+					{ "name": "strides", "value": [ 1, 1 ], "type": "Int32Array"}
+				] }
+			]
+		},
+		{
+			"name": "Convolution.Large2",
+			"attribute_maps": [
+				{ "attributes": [
+					{ "name": "auto_pad", "value": "NOTSET", "type": "String" },
+					{ "name": "dilations", "value": [ 2, 2 ], "type": "Int32Array" },
+					{ "name": "group", "value": 1, "type": "Int32" },
+					{ "name": "pads", "value": [ 0, 0, 0, 0 ], "type": "Int32Array" },
+					{ "name": "strides", "value": [ 2, 2 ], "type": "Int32Array"}
+				] }
+			]
+		},
+		{
+			"name": "ConvTranspose.Default",
+			"attribute_maps": [
+				{ "attributes": [
+					{ "name": "auto_pad", "value": "NOTSET", "type": "String" },
+					{ "name": "dilations", "value": [ 1, 1 ], "type": "Int32Array" },
+					{ "name": "group", "value": 1, "type": "Int32" },
+					{ "name": "pads", "value": [ 0, 0, 0, 0 ], "type": "Int32Array" },
+					{ "name": "strides", "value": [ 1, 1 ], "type": "Int32Array"}
+				] }
+			]
+		},
+		{
+			"name": "ConvTranspose.1D",
+			"attribute_maps": [
+				{ "attributes": [
+					{ "name": "auto_pad", "value": "NOTSET", "type": "String" },
+					{ "name": "dilations", "value": [ 1 ], "type": "Int32Array" },
+					{ "name": "group", "value": 1, "type": "Int32" },
+					{ "name": "pads", "value": [ 0, 0 ], "type": "Int32Array" },
+					{ "name": "strides", "value": [ 1 ], "type": "Int32Array"}
+				] }
+			]
+		},
+		{
+			"name": "ConvTranspose.3D",
+			"attribute_maps": [
+				{ "attributes": [
+					{ "name": "auto_pad", "value": "NOTSET", "type": "String" },
+					{ "name": "dilations", "value": [ 1, 1, 1 ], "type": "Int32Array" },
+					{ "name": "group", "value": 1, "type": "Int32" },
+					{ "name": "pads", "value": [ 0, 0, 0, 0, 0, 0 ], "type": "Int32Array" },
+					{ "name": "strides", "value": [ 1, 1, 1 ], "type": "Int32Array"}
+				] }
+			]
+		},
+		{
+			"name": "ConvTranspose.Attributes",
+			"attribute_maps": [
+				{ "attributes": [
+					{ "name": "auto_pad", "value": "NOTSET", "type": "String" },
+					{ "name": "dilations", "value": [ 1, 1 ], "type": "Int32Array" },
+					{ "name": "group", "value": 1, "type": "Int32" },
+					{ "name": "output_padding", "value": [ 1, 1 ], "type": "Int32Array" },
+					{ "name": "pads", "value": [ 0, 0, 0, 0 ], "type": "Int32Array" },
+					{ "name": "strides", "value": [ 3, 2 ], "type": "Int32Array"}
+				] }
+			]
+		},
+		{
+			"name": "ConvTranspose.Autopad.SameLower",
+			"attribute_maps": [
+				{ "attributes": [
+					{ "name": "auto_pad", "value": "SAME_LOWER", "type": "String" },
+					{ "name": "strides", "value": [ 2, 2 ], "type": "Int32Array"}
+				] }
+			]
+		},
+		{
+			"name": "ConvTranspose.Autopad.SameUpper",
+			"attribute_maps": [
+				{ "attributes": [
+					{ "name": "auto_pad", "value": "SAME_UPPER", "type": "String" },
+					{ "name": "strides", "value": [ 2, 2 ], "type": "Int32Array"}
+				] }
+			]
+		},
+		{
+			"name": "InstanceNormalization.Algorithms",
+			"attribute_maps": [
+				{ "attributes": [ { "name": "__UE__algorithm", "value": "Simple1x265", "type": "String" } ] },
+				{ "attributes": [ { "name": "__UE__algorithm", "value": "SharedMemory8x32", "type": "String" } ] },
+				{ "attributes": [ { "name": "__UE__algorithm", "value": "SharedMemory16x16", "type": "String" } ] },
+				{ "attributes": [ { "name": "__UE__algorithm", "value": "SharedMemory32x8", "type": "String" } ] }
+			]
+		},
+		{
+			"name": "PadOp.ModeAttributeSet",
+			"attribute_maps": [
+				{ "attributes": [ { "name": "mode", "value": "constant", "type": "String" } ] },
+				{ "attributes": [ { "name": "mode", "value": "reflect", "type": "String" } ] },
+				{ "attributes": [ { "name": "mode", "value": "edge", "type": "String" } ] }
+			]
+		},
+		{
+			"name": "PadOp.Pad00Rank1",
+			"multiply_with_attribute_sets": [ "PadOp.ModeAttributeSet" ],
+			"attribute_maps": [
+				{ "attributes": [
+					{ "name": "value", "value": 1, "type": "float" },
+					{ "name": "pads", "value": [ 0, 0 ], "type": "Int32Array" }
+				] }
+			]
+		},
+		{
+			"name": "PadOp.Pad03Rank1",
+			"multiply_with_attribute_sets": [ "PadOp.ModeAttributeSet" ],
+			"attribute_maps": [
+				{ "attributes": [
+					{ "name": "value", "value": 2, "type": "float" },
+					{ "name": "pads", "value": [ 0, 3 ], "type": "Int32Array" }
+				] }
+			]
+		},
+		{
+			"name": "PadOp.Pad20Rank1",
+			"multiply_with_attribute_sets": [ "PadOp.ModeAttributeSet" ],
+			"attribute_maps": [
+				{ "attributes": [
+					{ "name": "value", "value": 3, "type": "float" },
+					{ "name": "pads", "value": [ 2, 0 ], "type": "Int32Array" }
+				] }
+			]
+		},
+		{
+			"name": "PadOp.Pad41Rank1",
+			"multiply_with_attribute_sets": [ "PadOp.ModeAttributeSet" ],
+			"attribute_maps": [
+				{ "attributes": [
+					{ "name": "pads", "value": [ 4, 1 ], "type": "Int32Array" }
+				] }
+			]
+		},
+		{
+			"name": "PadOp.Pad00_10_12_34_Rank4",
+			"multiply_with_attribute_sets": [ "PadOp.ModeAttributeSet" ],
+			"attribute_maps": [
+				{ "attributes": [
+					{ "name": "pads", "value": [ 0, 1, 1, 3, 0, 0, 2, 4 ], "type": "Int32Array" }
+				] }
+			]
+		}
+	]
+
+}
diff --git a/Engine/Plugins/Experimental/NNE/Source/NNEQA/Resources/NNEQATestDescSchema.json b/Engine/Plugins/Experimental/NNE/Source/NNEQA/Resources/NNEQATestDescSchema.json
new file mode 100644
--- /dev/null
+++ b/Engine/Plugins/Experimental/NNE/Source/NNEQA/Resources/NNEQATestDescSchema.json
@@ -0,0 +1,306 @@
+{
+	"description": "schema for NNE test",
+	"type": "object",
+	"additionalProperties": false,
+	"properties": {
+		"$schema": {
+			"type": "string"
+		},
+		"model_test_categories": {
+			"type": "array",
+			"items": { "$ref": "#/$defs/test_config_category" }
+		},
+		"operator_test_categories": {
+			"type": "array",
+			"items": { "$ref": "#/$defs/test_config_category" }
+		},
+		"input_output_sets": {
+			"type": "array",
+			"items": { "$ref": "#/$defs/test_config_input_output_set" }
+		},
+		"attribute_sets": {
+			"type": "array",
+			"items": { "$ref": "#/$defs/test_config_attribute_set" }
+		},
+		"comment": {
+			"type": "string",
+			"description": "Unused by test system, usefull to add clarification in json file"
+		}
+	},
+
+	"$defs": {
+		"test_config_tensor": {
+			"type": "object",
+			"required": [ "shape" ],
+			"properties": {
+				"shape": {
+					"type": "array",
+					"items": { "type": "number" }
+			    },
+				"source": {
+					"type": "array",
+					"items": { "type": "string" }
+			    },
+				"type": {
+					"type": "string",
+					"enum": [ "float", "int32", "int64", "uint32", "boolean" ]
+				},
+				"comment": {
+					"type": "string",
+					"description": "Unused by test system, usefull to add clarification in json file"
+				}
+			}
+		},
+
+		"test_config_runtime": {
+			"type": "object",
+			"required": [ "name" ],
+			"additionalProperties": false,
+			"properties": {
+				"name": {
+					"type": "string",
+					"description": "name of the runtime to specialize for"
+				},
+				"skip": {
+					"type": "boolean",
+					"description": "true to skip the test for this runtime, false by default"
+				},
+				"absolute_tolerance": {
+					"type": "number",
+					"minimum": 0,
+					"description": "absolute tolerance of numerical tensor comparison routine"
+				},
+				"relative_tolerance": {
+					"type": "number",
+					"minimum": 0,
+					"maximum": 1,
+					"description": "relative tolerance of numerical tensor comparison routine"
+				},
+				"comment": {
+					"type": "string",
+					"description": "Unused by test system, usefull to add clarification in json file"
+				}
+			}
+		},
+
+		"test_config_dataset": {
+			"type": "object",
+			"required": [ "inputs" ],
+			"additionalProperties": false,
+			"properties": {
+				"inputs": {
+					"type": "array",
+					"description": "input to test",
+					"items": { "$ref": "#/$defs/test_config_tensor" }
+				},
+				"weights": {
+					"type": "array",
+					"description": "weights to test",
+					"items": { "$ref": "#/$defs/test_config_tensor" }
+				},
+				"outputs": {
+					"type": "array",
+					"description": "output to test",
+					"items": { "$ref": "#/$defs/test_config_tensor" }
+				},
+				"runtimes": {
+					"type": "array",
+					"description": "Array of operator to test",
+					"items": { "$ref": "#/$defs/test_config_runtime" }
+				},
+				"comment": {
+					"type": "string",
+					"description": "Unused by test system, usefull to add clarification in json file"
+				}
+			}
+		},
+
+		"test_config_attribute": {
+			"type": "object",
+			"required": [ "name", "value" ],
+			"additionalProperties": false,
+			"properties": {
+				"name": {
+					"type": "string",
+					"description": "Name of the attribute"
+				},
+				"value": {
+					"type": [ "number", "string", "array" ],
+					"description": "Value of the attribute"
+				},
+				"type": {
+					"type": "string",
+					"description": "Type of the attribute"
+				}
+			}
+		},
+
+		"test_config_attribute_map": {
+			"type": "object",
+			"required": [ "attributes" ],
+			"additionalProperties": false,
+			"properties": {
+				"attributes": {
+					"type": "array",
+					"description": "Attributes to test",
+					"items": { "$ref": "#/$defs/test_config_attribute" }
+				}
+			}
+		},
+
+		"test_config_target": {
+			"type": "object",
+			"required": [ "target" ],
+			"additionalProperties": false,
+			"properties": {
+				"target": {
+					"type": "string",
+					"description": "Name of the operator or model to run (for model .onnx should be ommited)"
+				},
+				"tags": {
+					"type": "array",
+					"description": "array of tags associated with the test, test can be filtered by tag when runned from command line",
+					"items": { "type": "string" }
+				},
+				"additional_datasets": {
+					"type": "array",
+					"description": "additional dataset associated with this target, by default all dataset witch name are a super set of this target category name will be runned, you can specify more dataset using this parameter",
+					"items": { "type": "string" }
+				},
+				"removed_datasets": {
+					"type": "array",
+					"description": "dataset that will be explicitely not run with this target, by default all dataset witch name are a super set of this target category name will be runned, you can filter out dataset using this parameter",
+					"items": { "type": "string" }
+				},
+				"skip": {
+					"type": "boolean",
+					"description": "true to skip this operator or model test entirely, false by default"
+				},
+				"absolute_tolerance": {
+					"type": "number",
+					"minimum": 0,
+					"description": "absolute tolerance (equality check)"
+				},
+				"relative_tolerance": {
+					"type": "number",
+					"minimum": 0,
+					"maximum": 1,
+					"description": "relative tolerance (equality check)"
+				},
+				"input_type": {
+					"type": "string",
+					"enum": [ "float", "int32", "int64", "uint32", "boolean" ]
+				},
+				"output_type": {
+					"type": "string",
+					"enum": [ "float", "int32", "int64", "uint32", "boolean" ]
+				},
+				"runtimes": {
+					"type": "array",
+					"description": "Array of runtimes to specialize test behavior for",
+					"items": { "$ref": "#/$defs/test_config_runtime" }
+				},
+				"comment": {
+					"type": "string",
+					"description": "Unused by test system, usefull to add clarification in json file"
+				}
+			}
+		},
+
+		"test_config_category": {
+			"type": "object",
+			"required": [ "category", "targets" ],
+			"additionalProperties": false,
+			"properties": {
+				"category": {
+					"type": "string",
+					"description": "Name of the category"
+				},
+				"targets": {
+					"type": "array",
+					"description": "Array of models or operators to test",
+					"items": { "$ref": "#/$defs/test_config_target" }
+				},
+				"skip": {
+					"type": "boolean",
+					"description": "true to skip this category entirely, false by default"
+				},
+				"is_model_test": {
+					"type": "boolean",
+					"description": "true if the category contain model tests, false for operators tests"
+				},
+				"runtimes": {
+					"type": "array",
+					"description": "Array of operator to test",
+					"items": { "$ref": "#/$defs/test_config_runtime" }
+				},
+				"additional_datasets": {
+					"type": "array",
+					"description": "additional dataset associated with this category, by default all dataset witch name are a super set of this category name will be runned, you can specify more dataset using this parameter",
+					"items": { "type": "string" }
+				},
+				"removed_datasets": {
+					"type": "array",
+					"description": "dataset that will be explicitely not run with this category, by default all dataset witch name are a super set of this category name will be runned, you can filter out dataset using this parameter",
+					"items": { "type": "string" }
+				},
+				"comment": {
+					"type": "string",
+					"description": "Unused by test system, usefull to add clarification in json file"
+				}
+			}
+		},
+
+		"test_config_input_output_set": {
+			"type": "object",
+			"required": [ "name", "datasets" ],
+			"additionalProperties": false,
+			"properties": {
+				"name": {
+					"type": "string",
+					"description": "Name of this dataset, any test name witch is a super set of this dataset name will use it"
+				},
+				"datasets": {
+					"type": "array",
+					"description": "Array of inputs/output to test",
+					"items": { "$ref": "#/$defs/test_config_dataset" }
+				},
+				"runtimes": {
+					"type": "array",
+					"description": "Array of runtimes to specialize test behavior for",
+					"items": { "$ref": "#/$defs/test_config_runtime" }
+				},
+				"comment": {
+					"type": "string",
+					"description": "Unused by test system, usefull to add clarification in json file"
+				}
+			}
+		},
+
+		"test_config_attribute_set": {
+			"type": "object",
+			"required": [ "name", "attribute_maps" ],
+			"additionalProperties": false,
+			"properties": {
+				"name": {
+					"type": "string",
+					"description": "Name of this attribute set"
+				},
+				"attribute_maps": {
+					"type": "array",
+					"description": "Array of attribute maps to test",
+					"items": { "$ref": "#/$defs/test_config_attribute_map" }
+				},
+				"multiply_with_attribute_sets": {
+					"type": "array",
+					"description": "Combine with other attribute sets",
+					"items": { "type": "string" }
+				},
+				"comment": {
+					"type": "string",
+					"description": "Unused by test system, usefull to add clarification in json file"
+				}
+			}
+		}
+	}
+}
\ No newline at end of file
diff --git a/Engine/Plugins/Experimental/NNE/Source/NNXQA/NNXQA.Build.cs b/Engine/Plugins/Experimental/NNE/Source/NNXQA/NNXQA.Build.cs
deleted file mode 100644
diff --git a/Engine/Plugins/Experimental/NNE/Source/NNXQA/Private/NNXQAJsonUtils.h b/Engine/Plugins/Experimental/NNE/Source/NNXQA/Private/NNXQAJsonUtils.h
deleted file mode 100644
diff --git a/Engine/Plugins/Experimental/NNE/Source/NNXQA/Private/NNXQAModule.cpp b/Engine/Plugins/Experimental/NNE/Source/NNXQA/Private/NNXQAModule.cpp
deleted file mode 100644
diff --git a/Engine/Plugins/Experimental/NNE/Source/NNXQA/Private/NNXQAParametricTest.cpp b/Engine/Plugins/Experimental/NNE/Source/NNXQA/Private/NNXQAParametricTest.cpp
deleted file mode 100644
diff --git a/Engine/Plugins/Experimental/NNE/Source/NNXQA/Private/NNXQAUtils.cpp b/Engine/Plugins/Experimental/NNE/Source/NNXQA/Private/NNXQAUtils.cpp
deleted file mode 100644
diff --git a/Engine/Plugins/Experimental/NNE/Source/NNXQA/Private/NNXQAUtils.h b/Engine/Plugins/Experimental/NNE/Source/NNXQA/Private/NNXQAUtils.h
deleted file mode 100644
diff --git a/Engine/Plugins/Experimental/NNE/Source/NNXQA/Public/NNXQAOperatorTest.h b/Engine/Plugins/Experimental/NNE/Source/NNXQA/Public/NNXQAOperatorTest.h
deleted file mode 100644
diff --git a/Engine/Plugins/Experimental/NNE/Source/NNXQA/Public/NNXQAParametricTest.h b/Engine/Plugins/Experimental/NNE/Source/NNXQA/Public/NNXQAParametricTest.h
deleted file mode 100644
diff --git a/Engine/Plugins/Experimental/NNE/Source/NNXQA/Resources/NNXQATestDesc.json b/Engine/Plugins/Experimental/NNE/Source/NNXQA/Resources/NNXQATestDesc.json
deleted file mode 100644
diff --git a/Engine/Plugins/Experimental/NNE/Source/NNXQA/Resources/NNXQATestDescSchema.json b/Engine/Plugins/Experimental/NNE/Source/NNXQA/Resources/NNXQATestDescSchema.json
deleted file mode 100644
