NNE - Upgrading the blueprinteable model to use the new API

diff --git a/Engine/Plugins/Experimental/NNE/Source/NNXCore/Private/NNECoreModel.cpp b/Engine/Plugins/Experimental/NNE/Source/NNXCore/Private/NNECoreModel.cpp
--- a/Engine/Plugins/Experimental/NNE/Source/NNXCore/Private/NNECoreModel.cpp
+++ b/Engine/Plugins/Experimental/NNE/Source/NNXCore/Private/NNECoreModel.cpp
@@ -2,9 +2,7 @@
 
 #include "NNECoreModel.h"
 
-#include "NNXCore.h"
-#include "NNXRuntimeFormat.h"
-#include "NNXInferenceModel.h"
+#include "NNECore.h"
 
 #include "Async/Async.h"
 #include "Async/TaskGraphInterfaces.h"
@@ -13,12 +11,16 @@
 
 TArray<FString> UNNEModel::GetRuntimeNames()
 {
-	using namespace NNX;
+	using namespace UE::NNECore;
+
 	TArray<FString> Result;
-	TArray<NNX::IRuntime*> Runtimes = GetAllRuntimes();
+	TArrayView<TWeakInterfacePtr<INNERuntime>> Runtimes = GetAllRuntimes();
 	for (int32 i = 0; i < Runtimes.Num(); i++)
 	{
-		Result.Add(Runtimes[i]->GetRuntimeName());
+		if (Runtimes[i].IsValid())
+		{
+			Result.Add(Runtimes[i]->GetRuntimeName());
+		}
 	}
 	return Result;
 }
@@ -37,40 +39,39 @@ bool UNNEModel::Load(FString RuntimeName, UNNEModelData* ModelData)
 {
 	check(IsInGameThread());
 
+	using namespace UE::NNECore;
+
 	if (IsAsyncRunning.IsValid() && *(IsAsyncRunning.Get()))
 	{
 		return false;
 	}
 
-	using namespace NNX;
-
-	IRuntime* Runtime = GetRuntime(RuntimeName);
-	if (!Runtime)
+	TWeakInterfacePtr<INNERuntimeCPU> Runtime = GetRuntime<INNERuntimeCPU>(RuntimeName);
+	if (!Runtime.IsValid())
 	{
-		UE_LOG(LogNNX, Error, TEXT("UNNEModel: No runtime '%s' found. Valid runtimes are: "), *RuntimeName);
-		TArray<NNX::IRuntime*> Runtimes = GetAllRuntimes();
-		for (int32 i = 0; i < Runtimes.Num(); i++)
+		UE_LOG(LogNNE, Error, TEXT("UNNEModel: No runtime '%s' found. Valid runtimes are: "), *RuntimeName);
+		TArray<FString> Names = GetRuntimeNames();
+		for (int32 i = 0; i < Names.Num(); i++)
 		{
-			UE_LOG(LogNNX, Error, TEXT("- %s"), *Runtimes[i]->GetRuntimeName());
+			UE_LOG(LogNNE, Error, TEXT("- %s"), *Names[i]);
 		}
 		return false;
 	}
 
 	if (!ModelData)
 	{
-		UE_LOG(LogNNX, Error, TEXT("UNNEModel: Valid model data required to load the model"));
+		UE_LOG(LogNNE, Error, TEXT("UNNEModel: Valid model data required to load the model"));
 		return false;
 	}
 
-	TConstArrayView<uint8> Data = ModelData->GetModelData(RuntimeName);
-	if (Data.Num() < 1)
+	TUniquePtr<IModelCPU> UniqueModel = Runtime->CreateModelCPU(ModelData);
+	if (!UniqueModel.IsValid())
 	{
-		UE_LOG(LogNNX, Error, TEXT("UNNEModel: No model data for %s found"), *RuntimeName);
+		UE_LOG(LogNNE, Error, TEXT("UNNEModel: Could not create the model"));
 		return false;
 	}
 
-	Model = TSharedPtr<NNX::FMLInferenceModel>(Runtime->CreateModel(Data).Release());
-
+	Model = TSharedPtr<IModelCPU>(UniqueModel.Release());
 	return Model.IsValid();
 }
 
@@ -102,12 +103,14 @@ TArray<int32> UNNEModel::GetInputShapes(int32 Index)
 {
 	check(IsInGameThread());
 
+	using namespace UE::NNECore;
+
 	if ((IsAsyncRunning.IsValid() && *(IsAsyncRunning.Get())) || !Model.IsValid())
 	{
 		return TArray<int32>();
 	}
 
-	TConstArrayView<NNX::FTensorDesc> Desc = Model->GetInputTensorDescs();
+	TConstArrayView<FTensorDesc> Desc = Model->GetInputTensorDescs();
 	if (Index < 0 || Index >= Desc.Num())
 	{
 		return TArray<int32>();
@@ -120,12 +123,14 @@ TArray<int32> UNNEModel::GetOutputShapes(int32 Index)
 {
 	check(IsInGameThread());
 
+	using namespace UE::NNECore;
+
 	if ((IsAsyncRunning.IsValid() && *(IsAsyncRunning.Get())) || !Model.IsValid())
 	{
 		return TArray<int32>();
 	}
 
-	TConstArrayView<NNX::FTensorDesc> Desc = Model->GetOutputTensorDescs();
+	TConstArrayView<FTensorDesc> Desc = Model->GetOutputTensorDescs();
 	if (Index < 0 || Index >= Desc.Num())
 	{
 		return TArray<int32>();
@@ -138,6 +143,8 @@ bool UNNEModel::SetInput(const TArray<FNNETensor>& Input)
 {
 	check(IsInGameThread());
 
+	using namespace UE::NNECore;
+
 	if ((IsAsyncRunning.IsValid() && *(IsAsyncRunning.Get())) || !Model.IsValid())
 	{
 		return false;
@@ -146,22 +153,25 @@ bool UNNEModel::SetInput(const TArray<FNNETensor>& Input)
 	InputBindings.Empty();
 	InputShapes.Empty();
 
-	TConstArrayView<NNX::FTensorDesc> InputDescs = Model->GetInputTensorDescs();
+	TConstArrayView<FTensorDesc> InputDescs = Model->GetInputTensorDescs();
 	if (InputDescs.Num() != Input.Num())
 	{
-		UE_LOG(LogNNX, Error, TEXT("UNNEModel: Invalid number of input tensors provided"));
+		UE_LOG(LogNNE, Error, TEXT("UNNEModel: Invalid number of input tensors provided"));
 		return false;
 	}
 
+	InputBindings.SetNum(Input.Num());
+	InputShapes.SetNum(Input.Num());
 	for (int32 i = 0; i < Input.Num(); i++)
 	{
-		InputBindings.Add(NNX::FMLTensorBinding::FromCPU((void*)(Input[i].Data.GetData()), sizeof(float) * Input[i].Data.Num()));
-		InputShapes.Add(NNX::FTensorShape::MakeFromSymbolic(UE::NNECore::FSymbolicTensorShape::Make(Input[i].Shape)));
+		InputBindings[i].Data = (void*)Input[i].Data.GetData();
+		InputBindings[i].SizeInBytes = Input[i].Data.Num() * sizeof(float);
+		InputShapes[i] = FTensorShape::MakeFromSymbolic(FSymbolicTensorShape::Make(Input[i].Shape));
 	}
 
 	if (Model->SetInputTensorShapes(InputShapes) != 0)
 	{
-		UE_LOG(LogNNX, Error, TEXT("UNNEModel: Failed to set input shapes"));
+		UE_LOG(LogNNE, Error, TEXT("UNNEModel: Failed to set input shapes"));
 		return false;
 	}
 
@@ -172,23 +182,26 @@ bool UNNEModel::RunSync(UPARAM(ref) TArray<FNNETensor>& Output)
 {
 	check(IsInGameThread());
 
+	using namespace UE::NNECore;
+
 	if ((IsAsyncRunning.IsValid() && *(IsAsyncRunning.Get())) || !Model.IsValid())
 	{
 		return false;
 	}
 
-	TArray<NNX::FMLTensorBinding> OutputBindings;
-
-	TConstArrayView<NNX::FTensorDesc> OutputDescs = Model->GetOutputTensorDescs();
+	TConstArrayView<FTensorDesc> OutputDescs = Model->GetOutputTensorDescs();
 	if (OutputDescs.Num() != Output.Num())
 	{
-		UE_LOG(LogNNX, Error, TEXT("UNNEModel: Invalid number of output tensors provided"));
+		UE_LOG(LogNNE, Error, TEXT("UNNEModel: Invalid number of output tensors provided"));
 		return false;
 	}
 
+	TArray<FTensorBindingCPU> OutputBindings;
+	OutputBindings.SetNum(Output.Num());
 	for (int32 i = 0; i < Output.Num(); i++)
 	{
-		OutputBindings.Add(NNX::FMLTensorBinding::FromCPU((void*)(Output[i].Data.GetData()), sizeof(float) * Output[i].Data.Num()));
+		OutputBindings[i].Data = (void*)Output[i].Data.GetData();
+		OutputBindings[i].SizeInBytes = Output[i].Data.Num() * sizeof(float);
 	}
 
 	return Model->RunSync(InputBindings, OutputBindings) == 0;
@@ -198,6 +211,8 @@ bool UNNEModel::RunAsync(FNNETaskPriority TaskPriority, FNNEModelOnAsyncResult O
 {
 	check(IsInGameThread());
 
+	using namespace UE::NNECore;
+
 	if ((IsAsyncRunning.IsValid() && *(IsAsyncRunning.Get())) ||  !Model.IsValid())
 	{
 		return false;
@@ -215,40 +230,47 @@ bool UNNEModel::RunAsync(FNNETaskPriority TaskPriority, FNNEModelOnAsyncResult O
 	for (int32 i = 0; i < InputBindings.Num(); i++)
 	{
 		InputData.Add(TArray<float>());
-		InputData[i].SetNumUninitialized((int32)(InputBindings[i].SizeInBytes / sizeof(float)));
-		FGenericPlatformMemory::Memcpy((void*)InputData[i].GetData(), InputBindings[i].CpuMemory, InputBindings[i].SizeInBytes);
+		InputData[i].SetNum((int32)(InputBindings[i].SizeInBytes / sizeof(float)));
+		FGenericPlatformMemory::Memcpy((void*)InputData[i].GetData(), InputBindings[i].Data, InputBindings[i].SizeInBytes);
 	}
 
 	// Launch a thread passing the model and a copy of the data to work on
-	TSharedPtr<NNX::FMLInferenceModel> ThreadModel = Model;
+	TSharedPtr<IModelCPU> ThreadModel = Model;
 	TSharedPtr<bool> LocalIsAsyncRunning = IsAsyncRunning;
 	AsyncTask(TaskPriority == FNNETaskPriority::Low ? ENamedThreads::AnyBackgroundThreadNormalTask : (TaskPriority == FNNETaskPriority::High ? ENamedThreads::AnyHiPriThreadHiPriTask : ENamedThreads::AnyNormalThreadNormalTask), [ThreadModel, ThreadInputData = MoveTemp(InputData), LocalIsAsyncRunning, OnAsyncResult]()
 	{
 		// Create input bindings
-		TArray<NNX::FMLTensorBinding> ThreadInputBindings;
+		TArray<FTensorBindingCPU> ThreadInputBindings;
+		ThreadInputBindings.SetNum(ThreadInputData.Num());
 		for (int32 i = 0; i < ThreadInputData.Num(); i++)
 		{
-			ThreadInputBindings.Add(NNX::FMLTensorBinding::FromCPU((void*)(ThreadInputData[i].GetData()), sizeof(float) * ThreadInputData[i].Num()));
+			ThreadInputBindings[i].Data = (void*)ThreadInputData[i].GetData();
+			ThreadInputBindings[i].SizeInBytes = ThreadInputData[i].Num() * sizeof(float);
 		}
 
 		// Create output data and bindings
 		TArray<FNNETensor> OutputTensors;
-		TArray<NNX::FMLTensorBinding> ThreadOutputBindings;
-		TConstArrayView<NNX::FTensorDesc> OutputDescs = ThreadModel->GetOutputTensorDescs();
+		TArray<FTensorBindingCPU> ThreadOutputBindings;
+		TConstArrayView<FTensorDesc> OutputDescs = ThreadModel->GetOutputTensorDescs();
+		ThreadOutputBindings.SetNum(OutputDescs.Num());
 		for (int32 i = 0; i < OutputDescs.Num(); i++)
 		{
 			OutputTensors.Add(FNNETensor());
 			OutputTensors[i].Shape = OutputDescs[i].GetShape().GetData();
+
 			int32 Size = OutputTensors[i].Shape.Num() > 0 ? 1 : 0;
 			for (int32 j = 0; j < OutputTensors[i].Shape.Num(); j++)
 			{
 				Size *= OutputTensors[i].Shape[j] > 0 ? OutputTensors[i].Shape[j] : 0;
 			}
+
 			if (Size > 0)
 			{
-				OutputTensors[i].Data.SetNumUninitialized(Size);
+				OutputTensors[i].Data.SetNum(Size);
 			}
-			ThreadOutputBindings.Add(NNX::FMLTensorBinding::FromCPU((void*)(OutputTensors[i].Data.GetData()), sizeof(float) * OutputTensors[i].Data.Num()));
+
+			ThreadOutputBindings[i].Data = (void*)OutputTensors[i].Data.GetData();
+			ThreadOutputBindings[i].SizeInBytes = OutputTensors[i].Data.Num() * sizeof(float);
 		}
 
 		bool bResult = ThreadModel->RunSync(ThreadInputBindings, ThreadOutputBindings) == 0;
diff --git a/Engine/Plugins/Experimental/NNE/Source/NNXCore/Public/NNECoreModel.h b/Engine/Plugins/Experimental/NNE/Source/NNXCore/Public/NNECoreModel.h
--- a/Engine/Plugins/Experimental/NNE/Source/NNXCore/Public/NNECoreModel.h
+++ b/Engine/Plugins/Experimental/NNE/Source/NNXCore/Public/NNECoreModel.h
@@ -4,7 +4,7 @@
 
 #include "CoreMinimal.h"
 
-#include "NNXRuntime.h"
+#include "NNECoreRuntimeCPU.h"
 #include "NNECoreModelData.h"
 
 #include "NNECoreModel.generated.h"
@@ -109,10 +109,10 @@ public:
 	
 private:
 
-	TSharedPtr<NNX::FMLInferenceModel> Model;
+	TSharedPtr<UE::NNECore::IModelCPU> Model;
 
-	TArray<NNX::FMLTensorBinding> InputBindings;
-	TArray<NNX::FTensorShape> InputShapes;
+	TArray<UE::NNECore::FTensorBindingCPU> InputBindings;
+	TArray<UE::NNECore::FTensorShape> InputShapes;
 
 	TSharedPtr<bool> IsAsyncRunning;
 };
