NNE : NNERuntimeRDG

* Update NNERuntimeRDG to follow NNE naming/filename convention

Test: nne smoketest green in editor win64

diff --git a/Engine/Plugins/Experimental/NNE/NNE.uplugin b/Engine/Plugins/Experimental/NNE/NNE.uplugin
--- a/Engine/Plugins/Experimental/NNE/NNE.uplugin
+++ b/Engine/Plugins/Experimental/NNE/NNE.uplugin
@@ -56,7 +56,7 @@
 			"PlatformAllowList": [ "Win64", "Linux", "Mac" ]
 		},
 		{
-			"Name": "NNXRuntimeRDG",
+			"Name": "NNERuntimeRDG",
 			"Type": "RuntimeAndProgram",
 			"LoadingPhase": "Default",
 			"PlatformAllowList": [ "Win64", "Linux", "Mac" ]
diff --git a/Engine/Plugins/Experimental/NNE/Source/NNEHlslShaders/Internal/NNEHlslShadersElementWiseBinaryCS.h b/Engine/Plugins/Experimental/NNE/Source/NNEHlslShaders/Internal/NNEHlslShadersElementWiseBinaryCS.h
--- a/Engine/Plugins/Experimental/NNE/Source/NNEHlslShaders/Internal/NNEHlslShadersElementWiseBinaryCS.h
+++ b/Engine/Plugins/Experimental/NNE/Source/NNEHlslShaders/Internal/NNEHlslShadersElementWiseBinaryCS.h
@@ -5,7 +5,7 @@
 #include "GlobalShader.h"
 #include "ShaderParameterUtils.h"
 #include "RenderGraphUtils.h"
-#include "NNXOperator.h"
+#include "NNECoreOperator.h"
 
 namespace UE::NNEHlslShaders::Internal
 {
@@ -22,7 +22,7 @@ namespace UE::NNEHlslShaders::Internal
 		DECLARE_GLOBAL_SHADER(TElementWiseBinaryCS);
 		SHADER_USE_PARAMETER_STRUCT(TElementWiseBinaryCS, FGlobalShader)
 
-		class FOperatorType : SHADER_PERMUTATION_ENUM_CLASS("OP_TYPENAME", EMLElementWiseBinaryOperatorType);
+		class FOperatorType : SHADER_PERMUTATION_ENUM_CLASS("OP_TYPENAME", NNECore::Internal::EElementWiseBinaryOperatorType);
 		class FBinaryNumDimensions : SHADER_PERMUTATION_RANGE_INT("NUM_DIMENSIONS", 1, FElementWiseBinaryConstants::MAX_NUM_DIMENSIONS);
 		using FPermutationDomain = TShaderPermutationDomain<FOperatorType, FBinaryNumDimensions>;
 
@@ -41,6 +41,6 @@ namespace UE::NNEHlslShaders::Internal
 
 	private:
 
-		static const FString GetOpFunc(EMLElementWiseBinaryOperatorType OpType);
+		static const FString GetOpFunc(NNECore::Internal::EElementWiseBinaryOperatorType OpType);
 	};
 } // UE::NNEHlslShaders::Internal
diff --git a/Engine/Plugins/Experimental/NNE/Source/NNEHlslShaders/Internal/NNEHlslShadersElementWiseUnaryCS.h b/Engine/Plugins/Experimental/NNE/Source/NNEHlslShaders/Internal/NNEHlslShadersElementWiseUnaryCS.h
--- a/Engine/Plugins/Experimental/NNE/Source/NNEHlslShaders/Internal/NNEHlslShadersElementWiseUnaryCS.h
+++ b/Engine/Plugins/Experimental/NNE/Source/NNEHlslShaders/Internal/NNEHlslShadersElementWiseUnaryCS.h
@@ -5,7 +5,7 @@
 #include "GlobalShader.h"
 #include "ShaderParameterUtils.h"
 #include "RenderGraphUtils.h"
-#include "NNXOperator.h"
+#include "NNECoreOperator.h"
 
 namespace UE::NNEHlslShaders::Internal
 {
@@ -21,7 +21,7 @@ namespace UE::NNEHlslShaders::Internal
 		DECLARE_GLOBAL_SHADER(TElementWiseUnaryCS);
 		SHADER_USE_PARAMETER_STRUCT(TElementWiseUnaryCS, FGlobalShader)
 
-		class FOperatorType : SHADER_PERMUTATION_ENUM_CLASS("OP_TYPENAME", EMLElementWiseUnaryOperatorType);
+		class FOperatorType : SHADER_PERMUTATION_ENUM_CLASS("OP_TYPENAME", NNECore::Internal::EElementWiseUnaryOperatorType);
 		using FPermutationDomain = TShaderPermutationDomain<FOperatorType>;
 
 	public:
@@ -40,6 +40,6 @@ namespace UE::NNEHlslShaders::Internal
 
 	private:
 
-		static const FString GetOpFunc(EMLElementWiseUnaryOperatorType OpType);
+		static const FString GetOpFunc(NNECore::Internal::EElementWiseUnaryOperatorType OpType);
 	};
 } // UE::NNEHlslShaders::Internal
diff --git a/Engine/Plugins/Experimental/NNE/Source/NNEHlslShaders/Internal/NNEHlslShadersElementWiseVariadicCS.h b/Engine/Plugins/Experimental/NNE/Source/NNEHlslShaders/Internal/NNEHlslShadersElementWiseVariadicCS.h
--- a/Engine/Plugins/Experimental/NNE/Source/NNEHlslShaders/Internal/NNEHlslShadersElementWiseVariadicCS.h
+++ b/Engine/Plugins/Experimental/NNE/Source/NNEHlslShaders/Internal/NNEHlslShadersElementWiseVariadicCS.h
@@ -5,7 +5,7 @@
 #include "GlobalShader.h"
 #include "ShaderParameterUtils.h"
 #include "RenderGraphUtils.h"
-#include "NNXOperator.h"
+#include "NNECoreOperator.h"
 
 namespace UE::NNEHlslShaders::Internal
 {
@@ -23,7 +23,7 @@ namespace UE::NNEHlslShaders::Internal
 		DECLARE_GLOBAL_SHADER(TElementWiseVariadicCS);
 		SHADER_USE_PARAMETER_STRUCT(TElementWiseVariadicCS, FGlobalShader)
 
-		class FOperatorType : SHADER_PERMUTATION_ENUM_CLASS("OP_TYPENAME", EMLElementWiseVariadicOperatorType);
+		class FOperatorType : SHADER_PERMUTATION_ENUM_CLASS("OP_TYPENAME", NNECore::Internal::EElementWiseVariadicOperatorType);
 		class FApplyScale : SHADER_PERMUTATION_BOOL("APPLYSCALE");
 		class FOutputAsInput : SHADER_PERMUTATION_BOOL("OUTPUTASINPUT");
 		class FNumInput : SHADER_PERMUTATION_RANGE_INT("NUMINPUT", 1, FElementWiseVariadicConstants::MAX_NUM_INPUT);
@@ -49,6 +49,6 @@ namespace UE::NNEHlslShaders::Internal
 
 	private:
 
-		static const FString GetOpFunc(EMLElementWiseVariadicOperatorType OpType);
+		static const FString GetOpFunc(NNECore::Internal::EElementWiseVariadicOperatorType OpType);
 	};
 } // UE::NNEHlslShaders::Internal
\ No newline at end of file
diff --git a/Engine/Plugins/Experimental/NNE/Source/NNEHlslShaders/Internal/NNEHlslShadersPadCS.h b/Engine/Plugins/Experimental/NNE/Source/NNEHlslShaders/Internal/NNEHlslShadersPadCS.h
--- a/Engine/Plugins/Experimental/NNE/Source/NNEHlslShaders/Internal/NNEHlslShadersPadCS.h
+++ b/Engine/Plugins/Experimental/NNE/Source/NNEHlslShaders/Internal/NNEHlslShadersPadCS.h
@@ -5,7 +5,6 @@
 #include "GlobalShader.h"
 #include "ShaderParameterUtils.h"
 #include "RenderGraphUtils.h"
-#include "NNXOperator.h"
 
 namespace UE::NNEHlslShaders::Internal
 {
diff --git a/Engine/Plugins/Experimental/NNE/Source/NNEHlslShaders/Internal/NNEHlslShadersUpsampleCS.h b/Engine/Plugins/Experimental/NNE/Source/NNEHlslShaders/Internal/NNEHlslShadersUpsampleCS.h
--- a/Engine/Plugins/Experimental/NNE/Source/NNEHlslShaders/Internal/NNEHlslShadersUpsampleCS.h
+++ b/Engine/Plugins/Experimental/NNE/Source/NNEHlslShaders/Internal/NNEHlslShadersUpsampleCS.h
@@ -5,7 +5,6 @@
 #include "GlobalShader.h"
 #include "ShaderParameterUtils.h"
 #include "RenderGraphUtils.h"
-#include "NNXOperator.h"
 
 namespace UE::NNEHlslShaders::Internal
 {
diff --git a/Engine/Plugins/Experimental/NNE/Source/NNEHlslShaders/Private/NNEHlslShadersElementWiseBinaryCS.cpp b/Engine/Plugins/Experimental/NNE/Source/NNEHlslShaders/Private/NNEHlslShadersElementWiseBinaryCS.cpp
--- a/Engine/Plugins/Experimental/NNE/Source/NNEHlslShaders/Private/NNEHlslShadersElementWiseBinaryCS.cpp
+++ b/Engine/Plugins/Experimental/NNE/Source/NNEHlslShaders/Private/NNEHlslShadersElementWiseBinaryCS.cpp
@@ -18,16 +18,16 @@ namespace UE::NNEHlslShaders::Internal
 		OutEnvironment.SetDefine(TEXT("ELEMENTWISE_OP(X,Y)"), *OpFunc);
 	}
 
-	const FString TElementWiseBinaryCS::GetOpFunc(EMLElementWiseBinaryOperatorType OpType)
+	const FString TElementWiseBinaryCS::GetOpFunc(NNECore::Internal::EElementWiseBinaryOperatorType OpType)
 	{
-		FString OpTable[(int32) EMLElementWiseBinaryOperatorType::MAX];
+		FString OpTable[(int32) NNECore::Internal::EElementWiseBinaryOperatorType::MAX];
 
-		for (int32 Idx = 0; Idx < (int32)EMLElementWiseBinaryOperatorType::MAX; ++Idx)
+		for (int32 Idx = 0; Idx < (int32) NNECore::Internal::EElementWiseBinaryOperatorType::MAX; ++Idx)
 		{
 			OpTable[Idx] = FString("");
 		}
 
-#define OP(OpName, OpFunc) OpTable[(int32) EMLElementWiseBinaryOperatorType::OpName] = OpFunc
+#define OP(OpName, OpFunc) OpTable[(int32) NNECore::Internal::EElementWiseBinaryOperatorType::OpName] = OpFunc
 		OP(Add,              TEXT("((X)+(Y))"));
 		//OP(And,            TEXT("((X)&&(Y))"));
 		OP(Div,              TEXT("((X)/(Y))"));
diff --git a/Engine/Plugins/Experimental/NNE/Source/NNEHlslShaders/Private/NNEHlslShadersElementWiseUnaryCS.cpp b/Engine/Plugins/Experimental/NNE/Source/NNEHlslShaders/Private/NNEHlslShadersElementWiseUnaryCS.cpp
--- a/Engine/Plugins/Experimental/NNE/Source/NNEHlslShaders/Private/NNEHlslShadersElementWiseUnaryCS.cpp
+++ b/Engine/Plugins/Experimental/NNE/Source/NNEHlslShaders/Private/NNEHlslShadersElementWiseUnaryCS.cpp
@@ -18,16 +18,16 @@ namespace UE::NNEHlslShaders::Internal
 		OutEnvironment.SetDefine(TEXT("ELEMENTWISE_OP(X)"), *OpFunc);
 	}
 
-	const FString TElementWiseUnaryCS::GetOpFunc(EMLElementWiseUnaryOperatorType OpType)
+	const FString TElementWiseUnaryCS::GetOpFunc(NNECore::Internal::EElementWiseUnaryOperatorType OpType)
 	{
-		FString OpTable[(int32) EMLElementWiseUnaryOperatorType::MAX];
+		FString OpTable[(int32) NNECore::Internal::EElementWiseUnaryOperatorType::MAX];
 
-		for (int32 Idx = 0; Idx < (int32) EMLElementWiseUnaryOperatorType::MAX; ++Idx)
+		for (int32 Idx = 0; Idx < (int32) NNECore::Internal::EElementWiseUnaryOperatorType::MAX; ++Idx)
 		{
 			OpTable[Idx] = FString("");
 		}
 
-#define OP(OpName, OpFunc) OpTable[(int32) EMLElementWiseUnaryOperatorType::OpName] = OpFunc
+#define OP(OpName, OpFunc) OpTable[(int32) NNECore::Internal::EElementWiseUnaryOperatorType::OpName] = OpFunc
 		OP(Abs,         TEXT("abs(X)"));
 		OP(Acos,        TEXT("acos(X)"));
 		OP(Acosh,       TEXT("acosh(X)"));
diff --git a/Engine/Plugins/Experimental/NNE/Source/NNEHlslShaders/Private/NNEHlslShadersElementWiseVariadicCS.cpp b/Engine/Plugins/Experimental/NNE/Source/NNEHlslShaders/Private/NNEHlslShadersElementWiseVariadicCS.cpp
--- a/Engine/Plugins/Experimental/NNE/Source/NNEHlslShaders/Private/NNEHlslShadersElementWiseVariadicCS.cpp
+++ b/Engine/Plugins/Experimental/NNE/Source/NNEHlslShaders/Private/NNEHlslShadersElementWiseVariadicCS.cpp
@@ -18,16 +18,16 @@ namespace UE::NNEHlslShaders::Internal
 		OutEnvironment.SetDefine(TEXT("ELEMENTWISE_OP(X,Y)"), *OpFunc);
 	}
 
-	const FString TElementWiseVariadicCS::GetOpFunc(EMLElementWiseVariadicOperatorType OpType)
+	const FString TElementWiseVariadicCS::GetOpFunc(NNECore::Internal::EElementWiseVariadicOperatorType OpType)
 	{
-		FString OpTable[(int32)EMLElementWiseVariadicOperatorType::MAX];
+		FString OpTable[(int32) NNECore::Internal::EElementWiseVariadicOperatorType::MAX];
 
-		for (int32 Idx = 0; Idx < (int32)EMLElementWiseVariadicOperatorType::MAX; ++Idx)
+		for (int32 Idx = 0; Idx < (int32) NNECore::Internal::EElementWiseVariadicOperatorType::MAX; ++Idx)
 		{
 			OpTable[Idx] = FString("");
 		}
 
-#define OP(OpName, OpFunc) OpTable[(int32) EMLElementWiseVariadicOperatorType::OpName] = OpFunc
+#define OP(OpName, OpFunc) OpTable[(int32) NNECore::Internal::EElementWiseVariadicOperatorType::OpName] = OpFunc
 		OP(Max, TEXT("max(X,Y)"));
 		OP(Min, TEXT("min(X,Y)"));
 		OP(Mean, TEXT("((X)+(Y))"));
diff --git a/Engine/Plugins/Experimental/NNE/Source/NNERuntimeRDG/Internal/NNERuntimeRDGNeuralPostProcessing.h b/Engine/Plugins/Experimental/NNE/Source/NNERuntimeRDG/Internal/NNERuntimeRDGNeuralPostProcessing.h
new file mode 100644
--- /dev/null
+++ b/Engine/Plugins/Experimental/NNE/Source/NNERuntimeRDG/Internal/NNERuntimeRDGNeuralPostProcessing.h
@@ -0,0 +1,83 @@
+// Copyright Epic Games, Inc. All Rights Reserved.
+
+#pragma once
+
+#include "CoreMinimal.h"
+#include "SceneViewExtension.h"
+
+#include "NNECoreModelData.h"
+#include "NNXInferenceModel.h"
+
+#include "NNERuntimeRDGNeuralPostProcessing.generated.h"
+
+class NNERUNTIMERDG_API FNNENeuralPostProcessing : public FSceneViewExtensionBase
+{
+public:
+
+	FNNENeuralPostProcessing(const FAutoRegister& AutoRegister);
+
+public:
+
+	int32 Add(FString RuntimeName, UNNEModelData* ModelData);
+	bool Remove(int32 ModelId);
+	bool SetWeight(int32 ModelId, float Weight);
+	bool SetRangeScale(int32 ModelId, float RangeScale);
+	bool SetInputSize(int32 ModelId, FIntPoint InputSize);
+	void Enable(int32 ModelId);
+	void Disable(int32 ModelId);
+
+public:
+
+	virtual void SetupViewFamily(FSceneViewFamily& InViewFamily) override {}
+	virtual void SetupView(FSceneViewFamily& InViewFamily, FSceneView& InView) override {}
+	virtual void BeginRenderViewFamily(FSceneViewFamily& InViewFamily) override {}
+	virtual void PreRenderViewFamily_RenderThread(FRHICommandListImmediate& RHICmdList, FSceneViewFamily& InViewFamily) override {}
+	virtual void PreRenderView_RenderThread(FRHICommandListImmediate& RHICmdList, FSceneView& InView) override {}
+	virtual void PrePostProcessPass_RenderThread(FRDGBuilder& GraphBuilder, const FSceneView& View, const FPostProcessingInputs& Inputs) override;
+
+private:
+
+	FCriticalSection CriticalSection;
+	int32 NumEnabled;
+	int32 LastId;
+	TMap<int32, TSharedPtr<NNX::FMLInferenceModel, ESPMode::ThreadSafe>> Models;
+	TMap<int32, float> Weights;
+	TMap<int32, float> RangeScales;
+	TMap<int32, FIntPoint> InputSizes;
+	TSet<int32> Enabled;
+
+};
+
+UCLASS(BlueprintType, Category = "NNE - Neural Network Engine")
+class NNERUNTIMERDG_API UNNENeuralPostProcessing : public UObject
+{
+	GENERATED_BODY()
+
+public:
+
+	UFUNCTION(BlueprintCallable, Category = "NNE - Neural Network Engine")
+	int32 Add(FString RuntimeName, UNNEModelData* ModelData);
+
+	UFUNCTION(BlueprintCallable, Category = "NNE - Neural Network Engine")
+	bool Remove(int32 ModelId);
+
+	UFUNCTION(BlueprintCallable, Category = "NNE - Neural Network Engine")
+	bool SetWeight(int32 ModelId, float Weight);
+
+	UFUNCTION(BlueprintCallable, Category = "NNE - Neural Network Engine")
+	bool SetRangeScale(int32 ModelId, float RangeScale);
+	
+	UFUNCTION(BlueprintCallable, Category = "NNE - Neural Network Engine")
+	bool SetInputSize(int32 ModelId, FIntPoint InputSize);
+
+	UFUNCTION(BlueprintCallable, Category = "NNE - Neural Network Engine")
+	void Enable(int32 ModelId);
+
+	UFUNCTION(BlueprintCallable, Category = "NNE - Neural Network Engine")
+	void Disable(int32 ModelId);
+
+private:
+
+	TSharedPtr<FNNENeuralPostProcessing> NeuralPostProcessing;
+
+};
\ No newline at end of file
diff --git a/Engine/Plugins/Experimental/NNE/Source/NNERuntimeRDG/NNERuntimeRDG.Build.cs b/Engine/Plugins/Experimental/NNE/Source/NNERuntimeRDG/NNERuntimeRDG.Build.cs
new file mode 100644
--- /dev/null
+++ b/Engine/Plugins/Experimental/NNE/Source/NNERuntimeRDG/NNERuntimeRDG.Build.cs
@@ -0,0 +1,52 @@
+// Copyright Epic Games, Inc. All Rights Reserved.
+
+using System.IO;
+using UnrealBuildTool;
+
+public class NNERuntimeRDG : ModuleRules
+{
+	public NNERuntimeRDG(ReadOnlyTargetRules Target) : base(Target)
+	{
+		PCHUsage = PCHUsageMode.UseExplicitOrSharedPCHs;
+
+		PrivateIncludePaths.AddRange(new string[] { Path.Combine(EngineDirectory, "Source/Runtime/Renderer/Private") });
+
+		PublicDependencyModuleNames.AddRange(new string[] 
+		{ 
+			"Core", 
+			"CoreUObject", 
+			"Engine", 
+			"InputCore",
+			"RenderCore"
+		});
+
+        PrivateDependencyModuleNames.AddRange(new string[]
+        {
+            "NNXCore",
+			"NNEUtils",
+			"NNEHlslShaders",
+            "RHI"
+        });
+
+        if (Target.Platform == UnrealTargetPlatform.Win64)
+		{
+			PrivateDependencyModuleNames.Add("D3D12RHI");
+			PrivateDependencyModuleNames.Add("DirectMLDefault");
+
+			AddEngineThirdPartyPrivateStaticDependencies(Target, "DX12");
+			AddEngineThirdPartyPrivateStaticDependencies(Target, "DirectMLDefault");
+
+			PublicDefinitions.Add("NNE_USE_DIRECTML");
+		}
+
+		if (Target.Platform == UnrealTargetPlatform.Mac)
+		{	
+			PrivateDependencyModuleNames.Add("MetalRHI");
+		}
+
+		if (Target.Platform == UnrealTargetPlatform.Linux)
+		{	
+			PrivateDependencyModuleNames.Add("VulkanRHI");
+		}
+	}
+}
diff --git a/Engine/Plugins/Experimental/NNE/Source/NNERuntimeRDG/Private/Hlsl/NNERuntimeRDGConv.cpp b/Engine/Plugins/Experimental/NNE/Source/NNERuntimeRDG/Private/Hlsl/NNERuntimeRDGConv.cpp
new file mode 100644
--- /dev/null
+++ b/Engine/Plugins/Experimental/NNE/Source/NNERuntimeRDG/Private/Hlsl/NNERuntimeRDGConv.cpp
@@ -0,0 +1,206 @@
+// Copyright Epic Games, Inc. All Rights Reserved.
+
+#include "NNERuntimeRDGConv.h"
+#include "NNEHlslShadersConvCS.h"
+#include "NNECoreAttributeMap.h"
+#include "NNERuntimeRDGHlslHelper.h"
+#include "NNECoreTensor.h"
+#include "NNECoreTypes.h"
+
+namespace UE::NNERuntimeRDG::Private::Hlsl
+{
+	DECLARE_GPU_STAT_NAMED(FNNEOperatorConv, TEXT("NNE.Operator.Hlsl.Conv"));
+
+	using EConvAutoPad = UE::NNEHlslShaders::Internal::EConvAutoPad;
+
+	/**
+	 * Convolution operator implementation
+	 */
+	class FConv : public FOperatorHlsl
+	{
+	public:
+
+		static FOperatorHlsl* Create()
+		{
+			return new FConv();
+		}
+
+		virtual ~FConv() = default;
+
+	private:
+
+		FConv() {}
+
+		int32 NumDimensions = 0;
+
+		EConvAutoPad AutoPad = EConvAutoPad::NOTSET;
+		TArray<int32> Dilations;
+		int32 Group = 1;
+		TArray<int32> Pads;
+		TArray<int32> Strides;
+
+	public:
+
+		virtual int PrepareOutputs(TConstArrayView<NNECore::Internal::FTensorRef> InputTensors, TArrayView<NNECore::Internal::FTensorRef> OutputTensors) const override
+		{
+			check(InputTensors.Num() >= 2 && InputTensors.Num() <= 3);
+			check(OutputTensors.Num() == 1);
+
+			const NNECore::FTensorShape& Input = InputTensors[0]->GetShape();
+			const NNECore::FTensorShape& Weights = InputTensors[1]->GetShape();
+
+			const TArray<int32> OutputShapeData = NNEHlslShaders::Internal::FConvCS::GetOutputShape(Input.GetData(), Weights.GetData(), AutoPad, Dilations, Strides, Pads);
+			const NNECore::FSymbolicTensorShape OutputShape = NNECore::FSymbolicTensorShape::Make(OutputShapeData);
+			
+			if (!OutputShape.IsConcrete())
+			{
+				return -1;
+			}
+			OutputTensors[0]->SetShape(NNECore::FTensorShape::MakeFromSymbolic(OutputShape));
+
+			return 0;
+		};
+
+		virtual bool Initialize(TConstArrayView<NNECore::FTensorDesc> InputTensorDescs, TConstArrayView<NNECore::FTensorDesc> OutputTensorDescs, const NNECore::FAttributeMap& Attributes) override
+		{
+            using namespace UE::NNEHlslShaders::Internal;
+
+			check(InputTensorDescs.Num() >= 2 && InputTensorDescs.Num() <= 3);
+			check(OutputTensorDescs.Num() == 1);
+
+            const NNECore::FTensorDesc& Input = InputTensorDescs[0];
+			const NNECore::FTensorDesc& Weights = InputTensorDescs[1];
+			const NNECore::FTensorDesc& Output = OutputTensorDescs[0];
+			
+			if (Input.GetShape().Rank() < 2)
+			{
+				UE_LOG(LogNNX, Warning, TEXT("Conv first input should be at least of rank 2"));
+				return false;
+			}
+			if (Weights.GetShape().Rank() != Input.GetShape().Rank())
+			{
+				UE_LOG(LogNNX, Warning, TEXT("Conv first and second inputs should be of same ranks"));
+				return false;
+			}
+			if (Output.GetShape().Rank() != Input.GetShape().Rank())
+			{
+				UE_LOG(LogNNX, Warning, TEXT("Conv first and output should be of same ranks"));
+				return false;
+			}
+
+			NumDimensions = Input.GetShape().Rank() - 2;
+
+			TArray<int32> DilationsOrStridesDefault;
+			DilationsOrStridesDefault.Init(1, NumDimensions);
+
+			FConvCS::LexFromString(AutoPad, *Attributes.GetValue<FString>(TEXT("auto_pad")));
+			Dilations = Attributes.GetValueOrDefault<TArray<int32>>(TEXT("dilations"), DilationsOrStridesDefault);
+			Group = Attributes.GetValueOrDefault<int>(TEXT("group"), 1);
+			if (AutoPad == EConvAutoPad::NOTSET)
+			{
+				TArray<int32> PadsDefault;
+				PadsDefault.Init(1, 2 * NumDimensions);
+
+				Pads = Attributes.GetValueOrDefault<TArray<int32>>(TEXT("pads"), PadsDefault);
+			}
+			Strides = Attributes.GetValueOrDefault<TArray<int32>>(TEXT("strides"), DilationsOrStridesDefault);
+
+			return true;
+		}
+
+		virtual void Dispatch(FRDGBuilder& GraphBuilder, TConstArrayView<FTensorRDGRef> InputTensors, TConstArrayView<FTensorRDGRef> OutputTensors) override
+		{
+			using namespace UE::NNEHlslShaders::Internal;
+
+			constexpr EConvAlgorithm Algorithm = EConvAlgorithm::SharedMemory;
+			constexpr EConvGroupSize GroupSize = EConvGroupSize::Size256;
+
+			check(InputTensors.Num() >= 2 && InputTensors.Num() <= 3);
+			check(OutputTensors.Num() == 1);
+			check(InputTensors[0] != nullptr);
+			check(InputTensors[1] != nullptr);
+			check(OutputTensors[0] != nullptr);
+			
+			const FTensorRDG& Input = *InputTensors[0];
+			const FTensorRDG& Weights = *InputTensors[1];
+			const FTensorRDG& Output = *OutputTensors[0];
+			const FTensorRDG* Bias = nullptr;
+			bool HasBias = false;
+
+			if (InputTensors.Num() == 3) {
+				HasBias = true;
+				check(InputTensors[2] != nullptr);
+				Bias = InputTensors[2];
+			}
+
+			check(Input.GetShape().Rank() > 2);
+			check(Weights.GetShape().Rank() == Input.GetShape().Rank());
+			check(Output.GetShape().Rank() == Input.GetShape().Rank());
+			check(NumDimensions == (Input.GetShape().Rank() - 2));
+
+			TArray<int32> OutputShape = FConvCS::GetOutputShape(Input.GetShape().GetData(), Weights.GetShape().GetData(), AutoPad, Dilations, Strides, Pads);
+
+			// Set parameters
+			FConvCS::FParameters* Params = GraphBuilder.AllocParameters<FConvCS::FParameters>();
+			FConvCS::FillInParameters(GroupSize, Input.GetShape().GetData(), Weights.GetShape().GetData(), HasBias, AutoPad, Group, Dilations,Strides, Pads, *Params);
+			Params->X = GraphBuilder.CreateSRV(FRDGBufferSRVDesc(Input.GetBuffer(), PF_R32_FLOAT));
+			Params->W = GraphBuilder.CreateSRV(FRDGBufferSRVDesc(Weights.GetBuffer(), PF_R32_FLOAT));
+			if (HasBias) {
+				Params->B = GraphBuilder.CreateSRV(FRDGBufferSRVDesc(Bias->GetBuffer(), PF_R32_FLOAT));
+			}
+			Params->Y = GraphBuilder.CreateUAV(FRDGBufferUAVDesc(Output.GetBuffer(), PF_R32_FLOAT));
+
+			FConvCS::FPermutationDomain PermutationVector;
+
+			PermutationVector.Set<FConvCS::FConvAlgorithm>(Algorithm);
+			PermutationVector.Set<FConvCS::FConvGroupSize>(GroupSize);
+			PermutationVector.Set<FConvCS::FConvNumDimensions>(NumDimensions);
+			PermutationVector.Set<FConvCS::FConvNumReadsPerThread>(FConvCS::GetNumReadsPerThread(GroupSize, Weights.GetShape().GetData(), Dilations, Strides));
+			PermutationVector.Set<FConvCS::FConvHasB>(HasBias);
+			TShaderMapRef<FConvCS> ComputeShader(GetGlobalShaderMap(GMaxRHIFeatureLevel), PermutationVector);
+
+			RDG_EVENT_SCOPE(GraphBuilder, "NNE.Operator.Hlsl.Conv");
+			RDG_GPU_STAT_SCOPE(GraphBuilder, FNNEOperatorConv);
+
+			FComputeShaderUtils::AddPass(
+				GraphBuilder,
+				RDG_EVENT_NAME("NNE.Operator.Hlsl.Conv.Dispatch"),
+				ERDGPassFlags::Compute | ERDGPassFlags::NeverCull,
+				ComputeShader,
+				Params,
+				FConvCS::GetGroupCount(OutputShape, FConvCS::GetGroupShape(GroupSize, NumDimensions)));
+		}
+	};
+
+	bool ValidateConvOperator(const NNECore::FAttributeMap& AttributeMap, TConstArrayView<ENNETensorDataType> InputTypes, TConstArrayView<NNECore::FSymbolicTensorShape> InputShapes)
+	{
+		bool bIsValid = true;
+
+		FAttributeValidator AttributeValidator;
+		AttributeValidator.AddOptional(TEXT("auto_pad"), ENNEAttributeDataType::String);
+		AttributeValidator.AddOptional(TEXT("dilations"), ENNEAttributeDataType::Int32Array);
+		AttributeValidator.AddOptional(TEXT("group"), ENNEAttributeDataType::Int32);
+		AttributeValidator.AddOptional(TEXT("kernel_shape"), ENNEAttributeDataType::Int32Array); // idea: cross check input weight shape with this attribute if present
+		AttributeValidator.AddOptional(TEXT("pads"), ENNEAttributeDataType::Int32Array);
+		AttributeValidator.AddOptional(TEXT("strides"), ENNEAttributeDataType::Int32Array);
+
+		bIsValid &= AttributeValidator.Validate(AttributeMap);
+
+		FInputValidator InputValidator;
+		InputValidator.AddSupportedType(ENNETensorDataType::Float);
+		InputValidator.AddRequired();
+		InputValidator.AddRequired();
+		InputValidator.AddOptional();
+		bIsValid &= InputValidator.Validate(InputTypes);
+
+		return bIsValid;
+	}
+
+	bool RegisterConvOperator(FOperatorRegistryHlsl& Registry)
+	{
+		Registry.OpAdd(TEXT("Conv"), FConv::Create, ValidateConvOperator);
+
+		return true;
+	}
+
+} // UE::NNERuntimeRDG::Private::Hlsl
diff --git a/Engine/Plugins/Experimental/NNE/Source/NNERuntimeRDG/Private/Hlsl/NNERuntimeRDGConv.h b/Engine/Plugins/Experimental/NNE/Source/NNERuntimeRDG/Private/Hlsl/NNERuntimeRDGConv.h
new file mode 100644
--- /dev/null
+++ b/Engine/Plugins/Experimental/NNE/Source/NNERuntimeRDG/Private/Hlsl/NNERuntimeRDGConv.h
@@ -0,0 +1,11 @@
+// Copyright Epic Games, Inc. All Rights Reserved.
+
+#pragma once
+
+#include "NNERuntimeRDG.h"
+#include "NNERuntimeRDGHlslOp.h"
+
+namespace UE::NNERuntimeRDG::Private::Hlsl
+{
+	bool RegisterConvOperator(FOperatorRegistryHlsl& Registry);
+} // UE::NNERuntimeRDG::Private::Hlsl
diff --git a/Engine/Plugins/Experimental/NNE/Source/NNERuntimeRDG/Private/Hlsl/NNERuntimeRDGConvTranspose.cpp b/Engine/Plugins/Experimental/NNE/Source/NNERuntimeRDG/Private/Hlsl/NNERuntimeRDGConvTranspose.cpp
new file mode 100644
--- /dev/null
+++ b/Engine/Plugins/Experimental/NNE/Source/NNERuntimeRDG/Private/Hlsl/NNERuntimeRDGConvTranspose.cpp
@@ -0,0 +1,213 @@
+// Copyright Epic Games, Inc. All Rights Reserved.
+
+#include "NNERuntimeRDGConvTranspose.h"
+#include "NNEHlslShadersConvTransposeCS.h"
+#include "NNERuntimeRDGHlslHelper.h"
+#include "NNECoreTensor.h"
+#include "NNECoreTypes.h"
+
+namespace UE::NNERuntimeRDG::Private::Hlsl
+{
+	DECLARE_GPU_STAT_NAMED(FNNEOperatorConvTranspose, TEXT("NNE.Operator.Hlsl.ConvTranspose"));
+
+	using EConvTransposeAutoPad = UE::NNEHlslShaders::Internal::EConvTransposeAutoPad;
+
+	/**
+	 * ConvTranspose operator implementation
+	 */
+	class FConvTranspose : public FOperatorHlsl
+	{
+	public:
+
+		static FOperatorHlsl* Create()
+		{
+			return new FConvTranspose();
+		}
+
+		virtual ~FConvTranspose() = default;
+
+	private:
+
+		FConvTranspose() {}
+
+		int32 NumDimensions = 0;
+
+		EConvTransposeAutoPad AutoPad = EConvTransposeAutoPad::NOTSET;
+		TArray<int32> Dilations;
+		int32 Group = 1;
+		TArray<int32> OutputPadding;
+		TArray<int32> Pads;
+		TArray<int32> Strides;
+
+	public:
+
+		virtual int PrepareOutputs(TConstArrayView<NNECore::Internal::FTensorRef> InputTensors, TArrayView<NNECore::Internal::FTensorRef> OutputTensors) const override
+		{
+			check(InputTensors.Num() >= 2 && InputTensors.Num() <= 3);
+			check(OutputTensors.Num() == 1);
+			
+			const NNECore::FTensorShape& Input = InputTensors[0]->GetShape();
+			const NNECore::FTensorShape& Weights = InputTensors[1]->GetShape();
+
+			
+			const TArray<int32> OutputShapeData = NNEHlslShaders::Internal::FConvTransposeCS::GetOutputShape(Input.GetData(), Weights.GetData(), AutoPad, Dilations, Strides, Pads, OutputPadding, Group);
+			const NNECore::FSymbolicTensorShape OutputShape = NNECore::FSymbolicTensorShape::Make(OutputShapeData);
+			if (!OutputShape.IsConcrete())
+			{
+				return -1;
+			}
+			OutputTensors[0]->SetShape(NNECore::FTensorShape::MakeFromSymbolic(OutputShape));
+
+			return 0;
+		};
+
+		virtual bool Initialize(TConstArrayView<NNECore::FTensorDesc> InputTensorDescs, TConstArrayView<NNECore::FTensorDesc> OutputTensorDescs, const NNECore::FAttributeMap& Attributes) override
+		{
+			using namespace UE::NNEHlslShaders::Internal;
+
+			check(InputTensorDescs.Num() >= 2 && InputTensorDescs.Num() <= 3);
+			check(OutputTensorDescs.Num() == 1);
+
+            const NNECore::FTensorDesc& Input = InputTensorDescs[0];
+			const NNECore::FTensorDesc& Weights = InputTensorDescs[1];
+			const NNECore::FTensorDesc& Output = OutputTensorDescs[0];
+
+			if (Input.GetShape().Rank() < 2)
+			{
+				UE_LOG(LogNNX, Warning, TEXT("ConvTranspose first input should be at least of rank 2"));
+				return false;
+			}
+			if (Weights.GetShape().Rank() != Input.GetShape().Rank())
+			{
+				UE_LOG(LogNNX, Warning, TEXT("ConvTranspose first and second inputs should be of same ranks"));
+				return false;
+			}
+			if (Output.GetShape().Rank() != Input.GetShape().Rank())
+			{
+				UE_LOG(LogNNX, Warning, TEXT("ConvTranspose first input and output should be of same ranks"));
+				return false;
+			}
+
+			NumDimensions = Input.GetShape().Rank() - 2;
+
+			TArray<int32> DilationsOrStridesDefault;
+			DilationsOrStridesDefault.Init(1, NumDimensions);
+
+			FConvTransposeCS::LexFromString(AutoPad, *Attributes.GetValue<FString>(TEXT("auto_pad")));
+			Dilations = Attributes.GetValueOrDefault<TArray<int32>>(TEXT("dilations"), DilationsOrStridesDefault);
+			Group = Attributes.GetValueOrDefault<int>(TEXT("group"), 1);
+			
+			TArray<int32> OutputPaddingDefault;
+			OutputPaddingDefault.Init(0, NumDimensions);
+			OutputPadding = Attributes.GetValueOrDefault<TArray<int32>>(TEXT("output_padding"), OutputPaddingDefault);
+			
+			if (AutoPad == EConvTransposeAutoPad::NOTSET)
+			{
+				TArray<int32> PadsDefault;
+				PadsDefault.Init(1, 2 * NumDimensions);
+
+				Pads = Attributes.GetValueOrDefault<TArray<int32>>(TEXT("pads"), PadsDefault);
+			}
+			Strides = Attributes.GetValueOrDefault<TArray<int32>>(TEXT("strides"), DilationsOrStridesDefault);
+
+			return true;
+		}
+
+		virtual void Dispatch(FRDGBuilder& GraphBuilder, TConstArrayView<FTensorRDGRef> InputTensors, TConstArrayView<FTensorRDGRef> OutputTensors) override
+		{
+			using namespace UE::NNEHlslShaders::Internal;
+
+			constexpr EConvTransposeAlgorithm Algorithm = EConvTransposeAlgorithm::SharedMemory;
+			constexpr EConvTransposeGroupSize GroupSize = EConvTransposeGroupSize::Size256;
+
+			check(InputTensors.Num() >= 2 && InputTensors.Num() <= 3);
+			check(OutputTensors.Num() == 1);
+			check(InputTensors[0] != nullptr);
+			check(InputTensors[1] != nullptr);
+			check(OutputTensors[0] != nullptr);
+
+			const FTensorRDG& Input = *InputTensors[0];
+			const FTensorRDG& Weights = *InputTensors[1];
+			const FTensorRDG& Output = *OutputTensors[0];
+			const FTensorRDG* Bias = nullptr;
+			bool HasBias = false;
+
+			if (InputTensors.Num() == 3) {
+				HasBias = true;
+				check(InputTensors[2] != nullptr);
+				Bias = InputTensors[2];
+			}
+
+			check(Input.GetShape().Rank() > 2);
+			check(Weights.GetShape().Rank() == Input.GetShape().Rank());
+			check(Output.GetShape().Rank() == Input.GetShape().Rank());
+			check(NumDimensions == (Input.GetShape().Rank() - 2));
+
+			TArray<int32> OutputShape = FConvTransposeCS::GetOutputShape(Input.GetShape().GetData(), Weights.GetShape().GetData(), AutoPad, Dilations, Strides, Pads, OutputPadding, Group);
+
+			// Set parameters
+			FConvTransposeCS::FParameters* Params = GraphBuilder.AllocParameters<FConvTransposeCS::FParameters>();
+			FConvTransposeCS::FillInParameters(GroupSize, Input.GetShape().GetData(), Weights.GetShape().GetData(), HasBias, AutoPad, Group, Dilations,Strides, Pads, OutputPadding, *Params);
+			Params->X = GraphBuilder.CreateSRV(FRDGBufferSRVDesc(Input.GetBuffer(), PF_R32_FLOAT));
+			Params->W = GraphBuilder.CreateSRV(FRDGBufferSRVDesc(Weights.GetBuffer(), PF_R32_FLOAT));
+			if (HasBias) {
+				Params->B = GraphBuilder.CreateSRV(FRDGBufferSRVDesc(Bias->GetBuffer(), PF_R32_FLOAT));
+			}
+			Params->Y = GraphBuilder.CreateUAV(FRDGBufferUAVDesc(Output.GetBuffer(), PF_R32_FLOAT));
+
+			FConvTransposeCS::FPermutationDomain PermutationVector;
+
+			PermutationVector.Set<FConvTransposeCS::FConvTransposeAlgorithm>(Algorithm);
+			PermutationVector.Set<FConvTransposeCS::FConvTransposeGroupSize>(GroupSize);
+			PermutationVector.Set<FConvTransposeCS::FConvTransposeNumStackDimensions>(NumDimensions);
+			PermutationVector.Set<FConvTransposeCS::FConvTransposeNumReadsPerThread>(FConvTransposeCS::GetNumReadsPerThread(GroupSize, Weights.GetShape().GetData(), Dilations, Strides));
+			PermutationVector.Set<FConvTransposeCS::FConvTransposeHasB>(HasBias);
+			TShaderMapRef<FConvTransposeCS> ComputeShader(GetGlobalShaderMap(GMaxRHIFeatureLevel), PermutationVector);
+
+			RDG_EVENT_SCOPE(GraphBuilder, "NNE.Operator.Hlsl.ConvTranspose");
+			RDG_GPU_STAT_SCOPE(GraphBuilder, FNNEOperatorConvTranspose);
+
+			FComputeShaderUtils::AddPass(
+				GraphBuilder,
+				RDG_EVENT_NAME("NNE.Operator.Hlsl.ConvTranspose.Dispatch"),
+				ERDGPassFlags::Compute | ERDGPassFlags::NeverCull,
+				ComputeShader,
+				Params,
+				FConvTransposeCS::GetGroupCount(OutputShape, FConvTransposeCS::GetGroupShape(GroupSize, NumDimensions)));		
+		}
+	};
+
+	bool ValidateConvTransposeOperator(const NNECore::FAttributeMap& AttributeMap, TConstArrayView<ENNETensorDataType> InputTypes, TConstArrayView<NNECore::FSymbolicTensorShape> InputShapes)
+	{
+		bool bIsValid = true;
+
+		FAttributeValidator AttributeValidator;
+		AttributeValidator.AddOptional(TEXT("auto_pad"), ENNEAttributeDataType::String);
+		AttributeValidator.AddOptional(TEXT("dilations"), ENNEAttributeDataType::Int32Array);
+		AttributeValidator.AddOptional(TEXT("group"), ENNEAttributeDataType::Int32);
+		//AttributeValidator.AddOptional(TEXT("kernel_shape"), ENNEAttributeDataType::Int32Array);
+		AttributeValidator.AddOptional(TEXT("output_padding"), ENNEAttributeDataType::Int32Array);
+		//AttributeValidator.AddOptional(TEXT("output_shape"), ENNEAttributeDataType::Int32Array);
+		AttributeValidator.AddOptional(TEXT("pads"), ENNEAttributeDataType::Int32Array);
+		AttributeValidator.AddOptional(TEXT("strides"), ENNEAttributeDataType::Int32Array);
+
+		bIsValid &= AttributeValidator.Validate(AttributeMap);
+
+		FInputValidator InputValidator;
+		InputValidator.AddSupportedType(ENNETensorDataType::Float);
+		InputValidator.AddRequired();
+		InputValidator.AddRequired();
+		InputValidator.AddOptional();
+		bIsValid &= InputValidator.Validate(InputTypes);
+
+		return bIsValid;
+	}
+
+	bool RegisterConvTransposeOperator(FOperatorRegistryHlsl& Registry)
+	{
+		Registry.OpAdd(TEXT("ConvTranspose"), FConvTranspose::Create, ValidateConvTransposeOperator);
+
+		return true;
+	}
+
+} // UE::NNERuntimeRDG::Private::Hlsl
diff --git a/Engine/Plugins/Experimental/NNE/Source/NNERuntimeRDG/Private/Hlsl/NNERuntimeRDGConvTranspose.h b/Engine/Plugins/Experimental/NNE/Source/NNERuntimeRDG/Private/Hlsl/NNERuntimeRDGConvTranspose.h
new file mode 100644
--- /dev/null
+++ b/Engine/Plugins/Experimental/NNE/Source/NNERuntimeRDG/Private/Hlsl/NNERuntimeRDGConvTranspose.h
@@ -0,0 +1,11 @@
+// Copyright Epic Games, Inc. All Rights Reserved.
+
+#pragma once
+
+#include "NNERuntimeRDG.h"
+#include "NNERuntimeRDGHlslOp.h"
+
+namespace UE::NNERuntimeRDG::Private::Hlsl
+{
+	bool RegisterConvTransposeOperator(FOperatorRegistryHlsl& Registry);
+} // UE::NNERuntimeRDG::Private::Hlsl
diff --git a/Engine/Plugins/Experimental/NNE/Source/NNERuntimeRDG/Private/Hlsl/NNERuntimeRDGElementWiseBinary.cpp b/Engine/Plugins/Experimental/NNE/Source/NNERuntimeRDG/Private/Hlsl/NNERuntimeRDGElementWiseBinary.cpp
new file mode 100644
--- /dev/null
+++ b/Engine/Plugins/Experimental/NNE/Source/NNERuntimeRDG/Private/Hlsl/NNERuntimeRDGElementWiseBinary.cpp
@@ -0,0 +1,163 @@
+// Copyright Epic Games, Inc. All Rights Reserved.
+
+#include "NNERuntimeRDGElementWiseBinary.h"
+#include "NNEHlslShadersElementWiseBinaryCS.h"
+#include "NNERuntimeRDGHlslHelper.h"
+#include "NNECoreTensor.h"
+#include "NNECoreTypes.h"
+
+namespace UE::NNERuntimeRDG::Private::Hlsl
+{
+	DECLARE_GPU_STAT_NAMED(FNNEOperatorElementWiseBinary, TEXT("NNE.Operator.Hlsl.ElementWise.Binary"));
+
+	using TElementWiseBinaryCS = typename UE::NNEHlslShaders::Internal::TElementWiseBinaryCS;
+	using FElementWiseBinaryConstants = UE::NNEHlslShaders::Internal::FElementWiseBinaryConstants;
+
+	/**
+	 * Binary element-wise operator implementation
+	 */
+	template<NNECore::Internal::EElementWiseBinaryOperatorType OpType>
+	class TElementWiseBinary : public FOperatorHlsl
+	{
+	public:
+
+		TElementWiseBinary() {}
+		virtual ~TElementWiseBinary() = default;
+
+	public:
+
+		virtual int PrepareOutputs(TConstArrayView<NNECore::Internal::FTensorRef> InputTensors, TArrayView<NNECore::Internal::FTensorRef> OutputTensors) const override
+		{
+			check(InputTensors.Num() == 2);
+			check(OutputTensors.Num() == 1);
+			
+			const NNECore::FTensorShape& LHSInput = InputTensors[0]->GetShape();
+			const NNECore::FTensorShape& RHSInput = InputTensors[1]->GetShape();
+			const int32 OutputRank = FMath::Max(LHSInput.Rank(), RHSInput.Rank());
+			TArray<uint32> OutputShapeData;
+			
+			OutputShapeData.SetNumUninitialized(OutputRank);
+			
+			for (int32 i = 0; i < OutputRank; ++i)
+			{
+				int32 LHSIndex = LHSInput.Rank() - 1 - i;
+				int32 RHSIndex = RHSInput.Rank() - 1 - i;
+				int32 LHSValue = LHSIndex >= 0 ? LHSInput.GetData()[LHSIndex] : 1;
+				int32 RHSValue = RHSIndex >= 0 ? RHSInput.GetData()[RHSIndex] : 1;
+				if (LHSValue != RHSValue && LHSValue != 1 && RHSValue != 1)
+				{
+					UE_LOG(LogNNX, Warning, TEXT("Error while computing shape for element wise binary op, input shapes are not compatible"));
+					return -1;
+				}
+				int32 OutputValue = FMath::Max(LHSValue, RHSValue);
+				OutputShapeData[OutputRank - 1 - i] = OutputValue;
+			}
+
+			NNECore::FTensorShape OutputShape = NNECore::FTensorShape::Make(OutputShapeData);
+			
+			OutputTensors[0]->SetShape(OutputShape);
+			
+			return 0;
+		}
+		
+		virtual bool Initialize(TConstArrayView<NNECore::FTensorDesc> InputTensorDescs, TConstArrayView<NNX::FTensorDesc> OutputTensorDescs, const NNECore::FAttributeMap& Attributes) override
+		{
+			check(InputTensorDescs.Num() == 2);
+			check(OutputTensorDescs.Num() == 1);
+		
+			return true;
+		}
+
+		virtual void Dispatch(FRDGBuilder& GraphBuilder, TConstArrayView<FTensorRDGRef> InputTensors, TConstArrayView<FTensorRDGRef> OutputTensors) override
+		{
+			check(InputTensors.Num() == 2);
+			check(OutputTensors.Num() == 1);
+			check(InputTensors[0] != nullptr);
+			check(InputTensors[1] != nullptr);
+			check(OutputTensors[0] != nullptr);
+			const FTensorRDG& LHSInput = *InputTensors[0];
+			const FTensorRDG& RHSInput = *InputTensors[1];
+			const FTensorRDG& Output = *OutputTensors[0];
+
+			FRDGBufferSRVRef LHSInputSRV = GraphBuilder.CreateSRV(FRDGBufferSRVDesc(LHSInput.GetBuffer(), PF_R32_FLOAT));
+			FRDGBufferSRVRef RHSInputSRV = GraphBuilder.CreateSRV(FRDGBufferSRVDesc(RHSInput.GetBuffer(), PF_R32_FLOAT));
+			FRDGBufferUAVRef OutputUAV = GraphBuilder.CreateUAV(FRDGBufferUAVDesc(Output.GetBuffer(), PF_R32_FLOAT));
+
+			FIntVector ThreadGroupCount = ComputeElementWiseThreadGroups(Output.GetVolume(), FElementWiseBinaryConstants::NUM_GROUP_THREADS);
+
+			// Set parameters
+			TElementWiseBinaryCS::FParameters* Params = GraphBuilder.AllocParameters<TElementWiseBinaryCS::FParameters>();
+			Params->LHSInput = LHSInputSRV;
+			Params->RHSInput = RHSInputSRV;
+			Params->Output = OutputUAV;
+			FillTensorStrideForBroadcastShaderParameters(LHSInput, Output.GetShape().Rank(), Params->TensorInfo, 0);
+			FillTensorStrideForBroadcastShaderParameters(RHSInput, Output.GetShape().Rank(), Params->TensorInfo, 1);
+			FillTensorStrideShaderParameters(Output, Params->TensorInfo, 2);
+			Params->Num = Output.GetVolume();
+			Params->ThreadCountX = ThreadGroupCount.X * FElementWiseBinaryConstants::NUM_GROUP_THREADS;
+
+			TElementWiseBinaryCS::FPermutationDomain PermutationVector;
+
+			PermutationVector.Set<TElementWiseBinaryCS::FOperatorType>(OpType);
+			PermutationVector.Set<TElementWiseBinaryCS::FBinaryNumDimensions>(Output.GetShape().Rank());
+
+			TShaderMapRef<TElementWiseBinaryCS> ComputeShader(GetGlobalShaderMap(GMaxRHIFeatureLevel), PermutationVector);
+		
+			RDG_EVENT_SCOPE(GraphBuilder, "NNE.Operator.Hlsl.ElementWise.Binary");
+			RDG_GPU_STAT_SCOPE(GraphBuilder, FNNEOperatorElementWiseBinary);
+
+			FComputeShaderUtils::AddPass(
+				GraphBuilder,
+				RDG_EVENT_NAME("NNE.Operator.Hlsl.ElementWise.Binary.Dispatch"),
+				ERDGPassFlags::Compute | ERDGPassFlags::NeverCull,
+				ComputeShader,
+				Params,
+				ThreadGroupCount);
+		}
+	};
+
+	bool ValidateElementWiseBinaryOperator(const NNECore::FAttributeMap& AttributeMap, TConstArrayView<ENNETensorDataType> InputTypes, TConstArrayView<NNECore::FSymbolicTensorShape> InputShapes)
+	{
+		bool bIsValid = true;
+
+		FAttributeValidator AttributeValidator;
+		bIsValid &= AttributeValidator.Validate(AttributeMap);
+
+		FInputValidator InputValidator;
+		InputValidator.AddSupportedType(ENNETensorDataType::Float);
+		InputValidator.AddRequired();
+		InputValidator.AddRequired();
+		bIsValid &= InputValidator.Validate(InputTypes);
+
+		return bIsValid;
+	}
+
+	template<NNECore::Internal::EElementWiseBinaryOperatorType OpType>
+	FOperatorHlsl* CreateElementWiseBinaryOperator()
+	{
+		return new TElementWiseBinary<OpType>();
+	}
+
+	bool RegisterElementWiseBinaryOperators(FOperatorRegistryHlsl& Registry)
+	{
+#define OP(Name) Registry.OpAdd(TEXT(#Name), CreateElementWiseBinaryOperator<NNECore::Internal::EElementWiseBinaryOperatorType::Name>, ValidateElementWiseBinaryOperator)
+		OP(Add);
+		//OP(And);
+		OP(Div);
+		//OP(Equal);
+		//OP(Greater);
+		//OP(GreaterOrEqual);
+		//OP(Less);
+		//OP(LessOrEqual);
+		OP(Mod);
+		OP(Mul);
+		//OP(Or);
+		OP(Prelu);
+		OP(Pow);
+		OP(Sub);
+		//OP(Or);
+#undef OP
+
+		return true;
+	}
+} // UE::NNERuntimeRDG::Private::Hlsl
diff --git a/Engine/Plugins/Experimental/NNE/Source/NNERuntimeRDG/Private/Hlsl/NNERuntimeRDGElementWiseBinary.h b/Engine/Plugins/Experimental/NNE/Source/NNERuntimeRDG/Private/Hlsl/NNERuntimeRDGElementWiseBinary.h
new file mode 100644
--- /dev/null
+++ b/Engine/Plugins/Experimental/NNE/Source/NNERuntimeRDG/Private/Hlsl/NNERuntimeRDGElementWiseBinary.h
@@ -0,0 +1,11 @@
+// Copyright Epic Games, Inc. All Rights Reserved.
+
+#pragma once
+
+#include "NNERuntimeRDG.h"
+#include "NNERuntimeRDGHlslOp.h"
+
+namespace UE::NNERuntimeRDG::Private::Hlsl
+{
+	bool RegisterElementWiseBinaryOperators(FOperatorRegistryHlsl& Registry);
+} // UE::NNERuntimeRDG::Private::Hlsl
diff --git a/Engine/Plugins/Experimental/NNE/Source/NNERuntimeRDG/Private/Hlsl/NNERuntimeRDGElementWiseUnary.cpp b/Engine/Plugins/Experimental/NNE/Source/NNERuntimeRDG/Private/Hlsl/NNERuntimeRDGElementWiseUnary.cpp
new file mode 100644
--- /dev/null
+++ b/Engine/Plugins/Experimental/NNE/Source/NNERuntimeRDG/Private/Hlsl/NNERuntimeRDGElementWiseUnary.cpp
@@ -0,0 +1,251 @@
+// Copyright Epic Games, Inc. All Rights Reserved.
+
+#include "NNERuntimeRDGElementWiseUnary.h"
+#include "NNEHlslShadersElementWiseUnaryCS.h"
+#include "NNERuntimeRDGHlslHelper.h"
+#include "NNECoreAttributeMap.h"
+#include "NNECoreTypes.h"
+#include "NNECoreTensor.h"
+
+namespace UE::NNERuntimeRDG::Private::Hlsl
+{
+	DECLARE_GPU_STAT_NAMED(FNNEOperatorElementWiseUnary, TEXT("NNE.Operator.Hlsl.ElementWise.Unary"));
+
+	using TElementWiseUnaryCS = typename UE::NNEHlslShaders::Internal::TElementWiseUnaryCS;
+	using FElementWiseUnaryConstants = UE::NNEHlslShaders::Internal::FElementWiseUnaryConstants;
+
+	/**
+	 * Unary element-wise operator implementation
+	 */
+	template<NNECore::Internal::EElementWiseUnaryOperatorType OpType>
+	class TElementWiseUnary : public FOperatorHlsl
+	{
+	public:
+
+		TElementWiseUnary() {}
+		virtual ~TElementWiseUnary() = default;
+
+	private:
+
+		float Alpha = 0.0f;
+		float Beta = 0.0f;
+		float Gamma = 0.0f;
+
+	public:
+
+		virtual int PrepareOutputs(TConstArrayView<NNECore::Internal::FTensorRef> InputTensors, TArrayView<NNECore::Internal::FTensorRef> OutputTensors) const override
+		{
+			check(InputTensors.Num() == 1);
+			check(OutputTensors.Num() == 1);
+			OutputTensors[0]->SetShape(InputTensors[0]->GetShape());
+			return 0;
+		}
+
+		virtual bool Initialize(TConstArrayView<NNECore::FTensorDesc> InputTensorDescs, TConstArrayView<NNECore::FTensorDesc> OutputTensorDescs, const NNECore::FAttributeMap& Attributes) override
+		{
+			check(InputTensorDescs.Num() == 1);
+			check(OutputTensorDescs.Num() == 1);
+
+			Alpha = Attributes.GetValueOrDefault(TEXT("alpha"), Alpha);
+			Beta = Attributes.GetValueOrDefault(TEXT("beta"), Beta);
+			Gamma = Attributes.GetValueOrDefault(TEXT("gamma"), Gamma);
+			return true;
+		}
+
+		virtual void Dispatch(FRDGBuilder& GraphBuilder, TConstArrayView<FTensorRDGRef> InInputTensors, TConstArrayView<FTensorRDGRef> InOutputTensors) override
+		{
+			check(InInputTensors[0] != nullptr);
+			check(InOutputTensors[0] != nullptr);
+			FRDGBufferSRVRef InputSRV = GraphBuilder.CreateSRV(FRDGBufferSRVDesc(InInputTensors[0]->GetBuffer(), PF_R32_FLOAT));
+			FRDGBufferUAVRef OutputUAV = GraphBuilder.CreateUAV(FRDGBufferUAVDesc(InOutputTensors[0]->GetBuffer(), PF_R32_FLOAT));
+		
+			int32 NumElements = InOutputTensors[0]->GetVolume();
+			FIntVector ThreadGroupCount = ComputeElementWiseThreadGroups(NumElements, FElementWiseUnaryConstants::NUM_GROUP_THREADS);
+
+			// Set parameters
+			TElementWiseUnaryCS::FParameters* Params = GraphBuilder.AllocParameters<TElementWiseUnaryCS::FParameters>();
+			Params->Input = InputSRV;
+			Params->Output = OutputUAV;
+			Params->Alpha = Alpha;
+			Params->Beta = Beta;
+			Params->Gamma = Gamma;
+			Params->Num = NumElements;
+			Params->ThreadCountX = ThreadGroupCount.X * FElementWiseUnaryConstants::NUM_GROUP_THREADS;
+
+			TElementWiseUnaryCS::FPermutationDomain PermutationVector;
+
+			PermutationVector.Set<TElementWiseUnaryCS::FOperatorType>(OpType);
+
+			TShaderMapRef<TElementWiseUnaryCS> ComputeShader(GetGlobalShaderMap(GMaxRHIFeatureLevel), PermutationVector);
+
+			RDG_EVENT_SCOPE(GraphBuilder, "NNE.Operator.Hlsl.ElementWise.Unary");
+			RDG_GPU_STAT_SCOPE(GraphBuilder, FNNEOperatorElementWiseUnary);
+		
+			FComputeShaderUtils::AddPass(
+				GraphBuilder,
+				RDG_EVENT_NAME("NNI.Operator.Hlsl.ElementWise.Unary.Dispatch"),
+				ERDGPassFlags::Compute | ERDGPassFlags::NeverCull,
+				ComputeShader,
+				Params,
+				ThreadGroupCount);
+		}
+	};
+
+	template<> TElementWiseUnary<NNECore::Internal::EElementWiseUnaryOperatorType::Selu>::TElementWiseUnary()
+		: Alpha(1.67326319217681884765625f), Beta(0.0f), Gamma(1.05070102214813232421875f)
+	{
+	}
+
+	template<> TElementWiseUnary<NNECore::Internal::EElementWiseUnaryOperatorType::Elu>::TElementWiseUnary()
+		: Alpha(1.0f), Beta(0.0f), Gamma(0.0f) 
+	{
+	}
+
+	template<> TElementWiseUnary<NNECore::Internal::EElementWiseUnaryOperatorType::HardSigmoid>::TElementWiseUnary()
+		: Alpha(0.2f), Beta(0.5f), Gamma(0.0f)
+	{
+	}
+
+	template<> TElementWiseUnary<NNECore::Internal::EElementWiseUnaryOperatorType::LeakyRelu>::TElementWiseUnary()
+		: Alpha(0.01f), Beta(0.0f), Gamma(0.0f)
+	{
+	}
+
+	template<NNECore::Internal::EElementWiseUnaryOperatorType OpType>
+	FOperatorHlsl* CreateElementWiseUnaryOperator()
+	{
+		return new TElementWiseUnary<OpType>();
+	}
+
+	template<NNECore::Internal::EElementWiseUnaryOperatorType OpType>
+	bool ValidateElementWiseUnaryOperator(const NNECore::FAttributeMap& AttributeMap, TConstArrayView<ENNETensorDataType> InputTypes, TConstArrayView<NNECore::FSymbolicTensorShape> InputShapes)
+	{
+		bool bIsValid = true;
+
+		FAttributeValidator AttributeValidator;
+		bIsValid &= AttributeValidator.Validate(AttributeMap);
+
+		FInputValidator InputValidator;
+		InputValidator.AddSupportedType(ENNETensorDataType::Float);
+		InputValidator.AddRequired();
+		bIsValid &= InputValidator.Validate(InputTypes);
+
+		return bIsValid;
+	}
+
+	template<>
+	bool ValidateElementWiseUnaryOperator<NNECore::Internal::EElementWiseUnaryOperatorType::Selu>(const NNECore::FAttributeMap& AttributeMap, TConstArrayView<ENNETensorDataType> InputTypes, TConstArrayView<NNECore::FSymbolicTensorShape> InputShapes)
+	{
+		bool bIsValid = true;
+
+		FAttributeValidator AttributeValidator;
+		AttributeValidator.AddOptional(TEXT("alpha"), ENNEAttributeDataType::Float);
+		AttributeValidator.AddOptional(TEXT("gamma"), ENNEAttributeDataType::Float);
+		bIsValid &= AttributeValidator.Validate(AttributeMap);
+
+		FInputValidator InputValidator;
+		InputValidator.AddSupportedType(ENNETensorDataType::Float);
+		InputValidator.AddRequired();
+		bIsValid &= InputValidator.Validate(InputTypes);
+
+		return bIsValid;
+	}
+
+	template<>
+	bool ValidateElementWiseUnaryOperator<NNECore::Internal::EElementWiseUnaryOperatorType::Elu>(const NNECore::FAttributeMap& AttributeMap, TConstArrayView<ENNETensorDataType> InputTypes, TConstArrayView<NNECore::FSymbolicTensorShape> InputShapes)
+	{
+		bool bIsValid = true;
+
+		FAttributeValidator AttributeValidator;
+		AttributeValidator.AddOptional(TEXT("alpha"), ENNEAttributeDataType::Float);
+		bIsValid &= AttributeValidator.Validate(AttributeMap);
+
+		FInputValidator InputValidator;
+		InputValidator.AddSupportedType(ENNETensorDataType::Float);
+		InputValidator.AddRequired();
+		bIsValid &= InputValidator.Validate(InputTypes);
+
+		return bIsValid;
+	}
+
+	template<>
+	bool ValidateElementWiseUnaryOperator<NNECore::Internal::EElementWiseUnaryOperatorType::HardSigmoid>(const NNECore::FAttributeMap& AttributeMap, TConstArrayView<ENNETensorDataType> InputTypes, TConstArrayView<NNECore::FSymbolicTensorShape> InputShapes)
+	{
+		bool bIsValid = true;
+
+		FAttributeValidator AttributeValidator;
+		AttributeValidator.AddOptional(TEXT("alpha"), ENNEAttributeDataType::Float);
+		AttributeValidator.AddOptional(TEXT("beta"), ENNEAttributeDataType::Float);
+		bIsValid &= AttributeValidator.Validate(AttributeMap);
+
+		FInputValidator InputValidator;
+		InputValidator.AddSupportedType(ENNETensorDataType::Float);
+		InputValidator.AddRequired();
+		bIsValid &= InputValidator.Validate(InputTypes);
+
+		return bIsValid;
+	}
+
+	template<>
+	bool ValidateElementWiseUnaryOperator<NNECore::Internal::EElementWiseUnaryOperatorType::LeakyRelu>(const NNECore::FAttributeMap& AttributeMap, TConstArrayView<ENNETensorDataType> InputTypes, TConstArrayView<NNECore::FSymbolicTensorShape> InputShapes)
+	{
+		bool bIsValid = true;
+
+		FAttributeValidator AttributeValidator;
+		AttributeValidator.AddOptional(TEXT("alpha"), ENNEAttributeDataType::Float);
+		bIsValid &= AttributeValidator.Validate(AttributeMap);
+
+		FInputValidator InputValidator;
+		InputValidator.AddSupportedType(ENNETensorDataType::Float);
+		InputValidator.AddRequired();
+		bIsValid &= InputValidator.Validate(InputTypes);
+
+		return bIsValid;
+	}
+	
+	bool RegisterElementWiseUnaryOperators(FOperatorRegistryHlsl& Registry)
+	{
+#define OP(Name) Registry.OpAdd(TEXT(#Name), CreateElementWiseUnaryOperator<NNECore::Internal::EElementWiseUnaryOperatorType::Name>, ValidateElementWiseUnaryOperator<NNECore::Internal::EElementWiseUnaryOperatorType::Name>)
+		OP(Abs);
+		OP(Acos);
+		OP(Acosh);
+		OP(Asin);
+		OP(Asinh);
+		OP(Atan);
+		OP(Atanh);
+		//OP(BitShift);
+		//OP(Cast);
+		OP(Ceil);
+		//OP(Clip);
+		OP(Cos);
+		OP(Cosh);
+		OP(Elu);
+		OP(Erf);
+		OP(Exp);
+		OP(Floor);
+		OP(IsInf);
+		OP(IsNan);
+		OP(HardSigmoid);
+		OP(HardSwish);
+		OP(LeakyRelu);
+		OP(Log);
+		OP(Neg);
+		//OP(Not);
+		OP(Reciprocal);
+		OP(Relu);
+		OP(Round);
+		OP(Selu);
+		OP(Sigmoid);
+		OP(Sign);
+		OP(Sin);
+		OP(Sinh);
+		OP(Softplus);
+		OP(Softsign);
+		OP(Sqrt);
+		OP(Tan);
+		OP(Tanh);
+#undef OP
+
+		return true;
+	}
+} // UE::NNERuntimeRDG::Private::Hlsl
diff --git a/Engine/Plugins/Experimental/NNE/Source/NNERuntimeRDG/Private/Hlsl/NNERuntimeRDGElementWiseUnary.h b/Engine/Plugins/Experimental/NNE/Source/NNERuntimeRDG/Private/Hlsl/NNERuntimeRDGElementWiseUnary.h
new file mode 100644
--- /dev/null
+++ b/Engine/Plugins/Experimental/NNE/Source/NNERuntimeRDG/Private/Hlsl/NNERuntimeRDGElementWiseUnary.h
@@ -0,0 +1,11 @@
+// Copyright Epic Games, Inc. All Rights Reserved.
+
+#pragma once
+
+#include "NNERuntimeRDG.h"
+#include "NNERuntimeRDGHlslOp.h"
+
+namespace UE::NNERuntimeRDG::Private::Hlsl
+{
+	bool RegisterElementWiseUnaryOperators(FOperatorRegistryHlsl& Registry);
+} // UE::NNERuntimeRDG::Private::Hlsl
diff --git a/Engine/Plugins/Experimental/NNE/Source/NNERuntimeRDG/Private/Hlsl/NNERuntimeRDGElementWiseVariadic.cpp b/Engine/Plugins/Experimental/NNE/Source/NNERuntimeRDG/Private/Hlsl/NNERuntimeRDGElementWiseVariadic.cpp
new file mode 100644
--- /dev/null
+++ b/Engine/Plugins/Experimental/NNE/Source/NNERuntimeRDG/Private/Hlsl/NNERuntimeRDGElementWiseVariadic.cpp
@@ -0,0 +1,229 @@
+// Copyright Epic Games, Inc. All Rights Reserved.
+
+#include "NNERuntimeRDGElementWiseVariadic.h"
+#include "NNEHlslShadersElementWiseVariadicCS.h"
+#include "NNERuntimeRDGHlslHelper.h"
+#include "NNECoreTensor.h"
+#include "NNECoreTypes.h"
+
+namespace UE::NNERuntimeRDG::Private::Hlsl
+{
+	DECLARE_GPU_STAT_NAMED(FNNEOperatorElementWiseVariadic, TEXT("NNE.Operator.Hlsl.ElementWise.Variadic"));
+
+	using TElementWiseVariadicCS = typename UE::NNEHlslShaders::Internal::TElementWiseVariadicCS;
+	using FElementWiseVariadicConstants = UE::NNEHlslShaders::Internal::FElementWiseVariadicConstants;
+
+	void AddOneVariadicOpPass(FRDGBuilder& GraphBuilder, 
+		TConstArrayView<FTensorRDGRef> InputTensors,
+		const FTensorRDG& OutputTensor,
+		bool OutputAsInput,
+		NNECore::Internal::EElementWiseVariadicOperatorType OpType,
+		float Scale)
+	{
+		static_assert(FElementWiseVariadicConstants::MAX_NUM_INPUT == 4, "This algorithm need to be adapted to math the shader.");
+		check(InputTensors.Num() <= FElementWiseVariadicConstants::MAX_NUM_INPUT);
+		check(InputTensors.Num() > 0);
+
+		// SRVs & UAV creations
+		FRDGBufferUAVRef OutputUAV = GraphBuilder.CreateUAV(FRDGBufferUAVDesc(OutputTensor.GetBuffer(), PF_R32_FLOAT));
+		FRDGBufferSRVRef InputsSRV[FElementWiseVariadicConstants::MAX_NUM_INPUT] = { nullptr, nullptr, nullptr, nullptr };
+
+		for (int32 i = 0; i < InputTensors.Num(); ++i)
+		{
+			check(InputTensors[i] != nullptr);
+			InputsSRV[i] = GraphBuilder.CreateSRV(FRDGBufferSRVDesc(InputTensors[i]->GetBuffer(), PF_R32_FLOAT));
+		}
+
+		// Set parameters
+		FIntVector ThreadGroupCount = ComputeElementWiseThreadGroups(OutputTensor.GetVolume(), FElementWiseVariadicConstants::NUM_GROUP_THREADS);
+		TElementWiseVariadicCS::FParameters* Params = GraphBuilder.AllocParameters<TElementWiseVariadicCS::FParameters>();
+
+		Params->Input0 = InputsSRV[0];
+		Params->Input1 = InputsSRV[1];
+		Params->Input2 = InputsSRV[2];
+		Params->Input3 = InputsSRV[3];
+		Params->Output = OutputUAV;
+		FillTensorStrideForBroadcastShaderParameters(*InputTensors[0], OutputTensor.GetShape().Rank(), Params->InputTensorInfo, 0);
+		if (InputTensors.Num() >= 2)
+		{
+			FillTensorStrideForBroadcastShaderParameters(*InputTensors[1], OutputTensor.GetShape().Rank(), Params->InputTensorInfo, 1);
+		}
+		if (InputTensors.Num() >= 3)
+		{
+			FillTensorStrideForBroadcastShaderParameters(*InputTensors[2], OutputTensor.GetShape().Rank(), Params->InputTensorInfo, 2);
+		}
+		if (InputTensors.Num() >= 4)
+		{
+			FillTensorStrideForBroadcastShaderParameters(*InputTensors[3], OutputTensor.GetShape().Rank(), Params->InputTensorInfo, 3);
+		}
+		FillTensorStrideShaderParameters(OutputTensor, Params->OutputTensorInfo, 0);
+		Params->Num = OutputTensor.GetVolume();
+		Params->ThreadCountX = ThreadGroupCount.X * FElementWiseVariadicConstants::NUM_GROUP_THREADS;
+		Params->Scale = Scale;
+
+		// Shader variation
+		TElementWiseVariadicCS::FPermutationDomain PermutationVector;
+
+		PermutationVector.Set<TElementWiseVariadicCS::FOperatorType>(OpType);
+		PermutationVector.Set<TElementWiseVariadicCS::FApplyScale>(Scale != 1.0f);
+		PermutationVector.Set<TElementWiseVariadicCS::FOutputAsInput>(OutputAsInput);
+		PermutationVector.Set<TElementWiseVariadicCS::FNumInput>(InputTensors.Num());
+		PermutationVector.Set<TElementWiseVariadicCS::FVariadicNumDimensions>(OutputTensor.GetShape().Rank());
+
+		// Add the pass to RDG
+		TShaderMapRef<TElementWiseVariadicCS> ComputeShader(GetGlobalShaderMap(GMaxRHIFeatureLevel), PermutationVector);
+		FComputeShaderUtils::AddPass(
+			GraphBuilder,
+			RDG_EVENT_NAME("NNI.Operator.Hlsl.ElementWise.Variadic.Dispatch"),
+			ERDGPassFlags::Compute | ERDGPassFlags::NeverCull,
+			ComputeShader,
+			Params,
+			ThreadGroupCount);
+	}
+
+	/**
+	 * Variadic Element-wise ML operator implementation
+	 */
+	template<NNECore::Internal::EElementWiseVariadicOperatorType OpType>
+	class TElementWiseVariadic : public FOperatorHlsl
+	{
+	public:
+
+		TElementWiseVariadic() {}
+		virtual ~TElementWiseVariadic() = default;
+
+	public:
+
+		virtual int PrepareOutputs(TConstArrayView<NNECore::Internal::FTensorRef> InputTensors, TArrayView<NNECore::Internal::FTensorRef> OutputTensors) const override
+		{
+			check(InputTensors.Num() > 0);
+			check(OutputTensors.Num() == 1);
+
+			const int32 NumInput = InputTensors.Num();
+			int32 OutputRank = 0;
+			for (int32 i = 0; i < NumInput; ++i)
+			{
+				OutputRank = FMath::Max(OutputRank, InputTensors[i]->GetShape().Rank());
+			}
+
+			TArray<uint32> OutputShapeData;
+
+			OutputShapeData.SetNumUninitialized(OutputRank);
+
+			for (int32 i = 0; i < OutputRank; ++i)
+			{
+				int32 OutputValue = 1;
+				for (int32 InputIdx = 0; InputIdx < NumInput; ++InputIdx)
+				{
+					int32 InputIndex = InputTensors[InputIdx]->GetShape().Rank() - 1 - i;
+					int32 InputValue = InputIndex >= 0 ? InputTensors[InputIdx]->GetShape().GetData()[InputIndex] : 1;
+					if (InputValue != OutputValue && InputValue != 1 && OutputValue != 1)
+					{
+						UE_LOG(LogNNX, Warning, TEXT("Error while computing shape for element wise variadic op, input shapes are not compatible"));
+						return -1;
+					}
+					OutputValue = FMath::Max(InputValue, OutputValue);
+				}
+				OutputShapeData[OutputRank - 1 - i] = OutputValue;
+			}
+
+			NNECore::FTensorShape OutputShape = NNECore::FTensorShape::Make(OutputShapeData);
+
+			OutputTensors[0]->SetShape(OutputShape);
+
+			return 0;
+		}
+		
+		virtual bool Initialize(TConstArrayView<NNX::FTensorDesc> InputTensorDescs, TConstArrayView<NNX::FTensorDesc> OutputTensorDescs, const NNECore::FAttributeMap& Attributes) override
+		{
+			check(InputTensorDescs.Num() > 0);
+			check(OutputTensorDescs.Num() == 1);
+			return true;
+		}
+
+		virtual void Dispatch(FRDGBuilder& GraphBuilder, TConstArrayView<FTensorRDGRef> InInputTensors, TConstArrayView<FTensorRDGRef> InOutputTensors) override
+		{
+			check(InInputTensors.Num() > 0);
+			check(InOutputTensors.Num() == 1);
+			check(InOutputTensors[0] != nullptr);
+
+			RDG_EVENT_SCOPE(GraphBuilder, "NNE.Operator.Hlsl.ElementWise.Variadic");
+			RDG_GPU_STAT_SCOPE(GraphBuilder, FNNEOperatorElementWiseVariadic);
+
+			FTensorRDGRef PassInputTensors[FElementWiseVariadicConstants::MAX_NUM_INPUT];
+			for (int32 InputOffset = 0; InputOffset < InInputTensors.Num(); InputOffset += FElementWiseVariadicConstants::MAX_NUM_INPUT)
+			{
+				uint32 NumInputLeftToHandle = InInputTensors.Num() - InputOffset;
+				uint32 NumInputForPass = FMath::Min(FElementWiseVariadicConstants::MAX_NUM_INPUT, NumInputLeftToHandle);
+				bool bIsFirstPass = (InputOffset == 0);
+				bool bIsLastPass = (NumInputLeftToHandle <= FElementWiseVariadicConstants::MAX_NUM_INPUT);
+				float Scale = 1.0f;
+
+				for (uint32 i = 0; i < NumInputForPass; ++i)
+				{
+					check(InInputTensors[InputOffset + i] != nullptr);
+					PassInputTensors[i] = InInputTensors[InputOffset + i];
+				}
+
+				if (OpType == NNECore::Internal::EElementWiseVariadicOperatorType::Mean && bIsLastPass)
+				{
+					Scale = 1.0f / InInputTensors.Num();
+				}
+			
+				AddOneVariadicOpPass(GraphBuilder,
+					MakeArrayView(PassInputTensors, NumInputForPass),
+					*InOutputTensors[0],
+					!bIsFirstPass,
+					OpType, Scale);
+			}
+		}
+	};
+
+	bool ValidateElementWiseVariadicOperator(const NNECore::FAttributeMap& AttributeMap, TConstArrayView<ENNETensorDataType> InputTypes, TConstArrayView<NNECore::FSymbolicTensorShape> InputShapes)
+	{
+		bool bIsValid = true;
+
+		FAttributeValidator AttributeValidator;
+		bIsValid &= AttributeValidator.Validate(AttributeMap);
+
+		if (InputTypes.Num() == 0)
+		{
+			UE_LOG(LogNNX, Error, TEXT("Element-wise variadic operator requires at least 1 input"));
+			bIsValid = false;
+		}
+		for (int32 i = 0; i < InputTypes.Num(); ++i)
+		{
+			if (InputTypes[i] != ENNETensorDataType::Float)
+			{
+				UE_LOG(LogNNX, Warning, TEXT("Element-wise variadic operator input '%d' of type '%d' is not supported, should be float at the moment."), i, InputTypes[i]);
+				bIsValid = false;
+			}
+		}
+		
+		FInputValidator InputValidator;
+		InputValidator.AddSupportedType(ENNETensorDataType::Float);
+		InputValidator.AddRequired();
+		InputValidator.AddRequired();
+		bIsValid &= InputValidator.Validate(InputTypes);
+
+		return bIsValid;
+	}
+
+	template<NNECore::Internal::EElementWiseVariadicOperatorType OpType>
+	FOperatorHlsl* CreateElementWiseVariadicOperator()
+	{
+		return new TElementWiseVariadic<OpType>();
+	}
+
+	bool RegisterElementWiseVariadicOperators(FOperatorRegistryHlsl& Registry)
+	{
+#define OP(Name) Registry.OpAdd(TEXT(#Name), CreateElementWiseVariadicOperator<NNECore::Internal::EElementWiseVariadicOperatorType::Name>)
+		OP(Max);
+		OP(Min);
+		OP(Mean);
+		OP(Sum);
+#undef OP
+
+		return true;
+	}
+} // UE::NNERuntimeRDG::Private::Hlsl
diff --git a/Engine/Plugins/Experimental/NNE/Source/NNERuntimeRDG/Private/Hlsl/NNERuntimeRDGElementWiseVariadic.h b/Engine/Plugins/Experimental/NNE/Source/NNERuntimeRDG/Private/Hlsl/NNERuntimeRDGElementWiseVariadic.h
new file mode 100644
--- /dev/null
+++ b/Engine/Plugins/Experimental/NNE/Source/NNERuntimeRDG/Private/Hlsl/NNERuntimeRDGElementWiseVariadic.h
@@ -0,0 +1,11 @@
+// Copyright Epic Games, Inc. All Rights Reserved.
+
+#pragma once
+
+#include "NNERuntimeRDG.h"
+#include "NNERuntimeRDGHlslOp.h"
+
+namespace UE::NNERuntimeRDG::Private::Hlsl
+{
+	bool RegisterElementWiseVariadicOperators(FOperatorRegistryHlsl& Registry);
+} // UE::NNERuntimeRDG::Private::Hlsl
diff --git a/Engine/Plugins/Experimental/NNE/Source/NNERuntimeRDG/Private/Hlsl/NNERuntimeRDGGather.cpp b/Engine/Plugins/Experimental/NNE/Source/NNERuntimeRDG/Private/Hlsl/NNERuntimeRDGGather.cpp
new file mode 100644
--- /dev/null
+++ b/Engine/Plugins/Experimental/NNE/Source/NNERuntimeRDG/Private/Hlsl/NNERuntimeRDGGather.cpp
@@ -0,0 +1,158 @@
+// Copyright Epic Games, Inc. All Rights Reserved.
+
+#include "NNERuntimeRDGGather.h"
+#include "NNERuntimeRDGHlslHelper.h"
+#include "NNEHlslShadersGatherCS.h"
+#include "NNECoreAttributeMap.h"
+#include "NNECoreTypes.h"
+#include "NNECoreTensor.h"
+
+namespace UE::NNERuntimeRDG::Private::Hlsl
+{
+	DECLARE_GPU_STAT_NAMED(FNNEOperatorGather, TEXT("NNE.Operator.Hlsl.Gather"));
+
+	/**
+	 * Gather operator implementation
+	 */
+	template <typename DataElementType, typename IndicesElementType>
+	class FGather : public FOperatorHlsl
+	{
+	public:
+
+		FGather() {}
+		virtual ~FGather() = default;
+
+	private:
+
+		int32 Axis = 0;
+
+	public:
+
+		virtual int PrepareOutputs(TConstArrayView<NNECore::Internal::FTensorRef> InputTensors, TArrayView<NNECore::Internal::FTensorRef> OutputTensors) const override
+		{
+			check(InputTensors.Num() == 2)
+			check(OutputTensors.Num() == 1)
+			UE_LOG(LogNNX, Warning, TEXT("Gather shape inference is not implemented at the moment"));
+			return -1;
+		};
+		
+		virtual bool Initialize(TConstArrayView<NNECore::FTensorDesc> InputTensorDescs, TConstArrayView<NNECore::FTensorDesc> OutputTensorDescs, const NNECore::FAttributeMap& Attributes) override
+		{
+			const int32 MaxNumDimensions = NNEHlslShaders::Internal::FGatherConstants::MAX_NUM_DIMENSIONS;
+			
+			check(InputTensorDescs.Num() == 2)
+			check(OutputTensorDescs.Num() == 1)
+
+			const NNECore::FTensorDesc& Data = InputTensorDescs[0];
+			const NNECore::FTensorDesc& Indices = InputTensorDescs[1];
+			const NNECore::FTensorDesc& Output = OutputTensorDescs[0];
+
+			if (Output.GetShape().Rank() <= MaxNumDimensions)
+			{
+				UE_LOG(LogNNX, Warning, TEXT("Gather first input should be of rank %d or less"), MaxNumDimensions);
+				return false;
+			}
+			if (Data.GetShape().Rank() == 0)
+			{
+				UE_LOG(LogNNX, Warning, TEXT("Gather first input should be at least of rank 1"));
+				return false;
+			}
+			if (Indices.GetShape().Rank() == 0)
+			{
+				UE_LOG(LogNNX, Warning, TEXT("Gather second input should be at least of rank 1"));
+				return false;
+			}
+			if ((Data.GetShape().Rank() + Indices.GetShape().Rank() - 1) > MaxNumDimensions)
+			{
+				UE_LOG(LogNNX, Warning, TEXT("Gather sum of input 0 and 1 ranks -1 should be less than %d"), MaxNumDimensions);
+				return false;
+			}
+
+			Axis = Attributes.GetValueOrDefault(TEXT("axis"), Axis);
+			if (Axis >= Data.GetShape().Rank())
+			{
+				UE_LOG(LogNNX, Warning, TEXT("Gather Axis attribute should be inferior to first input rank"));
+				return false;
+			}
+			if (Axis < -Data.GetShape().Rank())
+{
+				UE_LOG(LogNNX, Warning, TEXT("Gather Axis attribute should be superior or equal to minus the first input rank"));
+				return false;
+			}
+			Axis = Axis >= 0 ? Axis : Data.GetShape().Rank() + Axis;
+
+			return true;
+		}
+
+		virtual void Dispatch(FRDGBuilder& GraphBuilder, TConstArrayView<FTensorRDGRef> InputTensors, TConstArrayView<FTensorRDGRef> OutputTensors) override
+		{
+			using namespace UE::NNEHlslShaders::Internal;
+
+			check(InputTensors.Num() == 2)
+			check(OutputTensors.Num() == 1)
+			check(InputTensors[0] != nullptr);
+			check(InputTensors[1] != nullptr);
+			check(OutputTensors[0] != nullptr);
+			check(OutputTensors[0]->GetShape().Rank() <= FGatherConstants::MAX_NUM_DIMENSIONS)
+			check(InputTensors[0]->GetShape().Rank() > 0)
+			check(InputTensors[1]->GetShape().Rank() > 0)
+			check(InputTensors[0]->GetShape().Rank() + (InputTensors[1]->GetShape().Rank() - 1) <= FGatherConstants::MAX_NUM_DIMENSIONS)
+
+			const FTensorRDG& Data = *InputTensors[0];
+			const FTensorRDG& Indices = *InputTensors[1];
+			const FTensorRDG& Output = *OutputTensors[0];
+
+			// Set parameters
+			TGatherCS::FParameters* Parameters = GraphBuilder.AllocParameters<TGatherCS::FParameters>();
+			TGatherCS::FillInParameters(Axis, Data, Indices, *Parameters);
+			Parameters->Data = GraphBuilder.CreateSRV(FRDGBufferSRVDesc(Data.GetBuffer(), PF_R32_FLOAT));
+			Parameters->Indices = GraphBuilder.CreateSRV(FRDGBufferSRVDesc(Indices.GetBuffer(), PF_R32_FLOAT));
+			Parameters->Output = GraphBuilder.CreateUAV(FRDGBufferUAVDesc(Output.GetBuffer(), PF_R32_FLOAT));
+
+			TGatherCS::FPermutationDomain PermutationVector;
+			PermutationVector.Set<TGatherCS::FGatherNumOutputDimensions>(Output.GetShape().Rank());
+			TShaderMapRef<TGatherCS> ComputeShader(GetGlobalShaderMap(GMaxRHIFeatureLevel), PermutationVector);
+
+			FIntVector ThreadGroupCount = TGatherCS::GetGroupCount(*Parameters);
+
+			RDG_EVENT_SCOPE(GraphBuilder, "NNE.Operator.Hlsl.Gather");
+			RDG_GPU_STAT_SCOPE(GraphBuilder, FNNEOperatorGather);
+
+			FComputeShaderUtils::AddPass(
+				GraphBuilder,
+				RDG_EVENT_NAME("NNI.Operator.Hlsl.Gather.Dispatch"),
+				ERDGPassFlags::Compute | ERDGPassFlags::NeverCull,
+				ComputeShader,
+				Parameters,
+				ThreadGroupCount);
+		}
+	};
+
+	bool ValidateGatherOperator(const NNECore::FAttributeMap& AttributeMap, TConstArrayView<ENNETensorDataType> InputTypes, TConstArrayView<NNECore::FSymbolicTensorShape> InputShapes)
+	{
+		bool bIsValid = true;
+
+		FAttributeValidator AttributeValidator;
+		AttributeValidator.AddOptional(TEXT("axis"), ENNEAttributeDataType::Int32);
+		bIsValid &= AttributeValidator.Validate(AttributeMap);
+
+		FInputValidator InputValidator;
+		InputValidator.AddSupportedType(ENNETensorDataType::Float);
+		InputValidator.AddRequired();
+		InputValidator.AddRequired();//Indices should be int32 or int64
+		bIsValid &= InputValidator.Validate(InputTypes);
+
+		return bIsValid;
+	}
+
+	FOperatorHlsl* CreateGatherOperator()
+	{
+		return new FGather<float, int32>();
+	}
+
+	bool RegisterGatherOperator(FOperatorRegistryHlsl& Registry)
+	{
+		Registry.OpAdd(TEXT("Gather"), CreateGatherOperator, ValidateGatherOperator);
+		return true;
+	}
+} // UE::NNERuntimeRDG::Private::Hlsl
\ No newline at end of file
diff --git a/Engine/Plugins/Experimental/NNE/Source/NNERuntimeRDG/Private/Hlsl/NNERuntimeRDGGather.h b/Engine/Plugins/Experimental/NNE/Source/NNERuntimeRDG/Private/Hlsl/NNERuntimeRDGGather.h
new file mode 100644
--- /dev/null
+++ b/Engine/Plugins/Experimental/NNE/Source/NNERuntimeRDG/Private/Hlsl/NNERuntimeRDGGather.h
@@ -0,0 +1,11 @@
+// Copyright Epic Games, Inc. All Rights Reserved.
+
+#pragma once
+
+#include "NNERuntimeRDG.h"
+#include "NNERuntimeRDGHlslOp.h"
+
+namespace UE::NNERuntimeRDG::Private::Hlsl
+{
+	bool RegisterGatherOperator(FOperatorRegistryHlsl& Registry);
+} // UE::NNERuntimeRDG::Private::Hlsl
diff --git a/Engine/Plugins/Experimental/NNE/Source/NNERuntimeRDG/Private/Hlsl/NNERuntimeRDGGemm.cpp b/Engine/Plugins/Experimental/NNE/Source/NNERuntimeRDG/Private/Hlsl/NNERuntimeRDGGemm.cpp
new file mode 100644
--- /dev/null
+++ b/Engine/Plugins/Experimental/NNE/Source/NNERuntimeRDG/Private/Hlsl/NNERuntimeRDGGemm.cpp
@@ -0,0 +1,193 @@
+// Copyright Epic Games, Inc. All Rights Reserved.
+
+#include "NNERuntimeRDGGemm.h"
+#include "NNEHlslShadersGemmCS.h"
+#include "NNERuntimeRDGHlslHelper.h"
+#include "NNECoreAttributeMap.h"
+#include "NNECoreTensor.h"
+#include "NNECoreTypes.h"
+
+namespace UE::NNERuntimeRDG::Private::Hlsl
+{
+	DECLARE_GPU_STAT_NAMED(FNNEOperatorGemm, TEXT("NNE.Operator.Hlsl.Gemm"));
+
+	/**
+	 * Gemm operator implementation
+	 */
+	class TGemm : public FOperatorHlsl
+	{
+	public:
+
+		TGemm() {}
+		virtual ~TGemm() = default;
+
+	private:
+
+		float InputAlpha = 1.0f;
+		float InputBeta = 1.0f;
+		int32 InputTransA = 0;
+		int32 InputTransB = 0;
+
+		bool bIsCScalar = false;
+		bool bNoBias = true;
+
+	public:
+
+		virtual int PrepareOutputs(TConstArrayView<NNECore::Internal::FTensorRef> InputTensors, TArrayView<NNECore::Internal::FTensorRef> OutputTensors) const override
+		{
+			check(InputTensors.Num() >= 2 && InputTensors.Num() <= 3);
+			check(OutputTensors.Num() == 1);
+
+			const NNX::FTensorShape& InputA = InputTensors[0]->GetShape();
+			const NNX::FTensorShape& InputB = InputTensors[1]->GetShape();
+			if (InputA.Rank() != 2 || InputB.Rank() != 2)
+			{
+				return -1;
+			}
+
+			uint32 M = InputTransA != 0 ? InputA.GetData()[1] : InputA.GetData()[0];
+			uint32 N = InputTransB != 0 ? InputB.GetData()[0] : InputB.GetData()[1];
+			TArray<uint32> OutputShapeData;
+
+			OutputShapeData.Emplace(M);
+			OutputShapeData.Emplace(N);
+			
+			NNECore::FTensorShape OutputShape = NNECore::FTensorShape::Make(OutputShapeData);
+			
+			OutputTensors[0]->SetShape(OutputShape);
+			return 0;
+		};
+		
+		virtual bool Initialize(TConstArrayView<NNECore::FTensorDesc> InputTensorDescs, TConstArrayView<NNECore::FTensorDesc> OutputTensorDescs, const NNECore::FAttributeMap& Attributes) override
+		{
+			check(InputTensorDescs.Num() >= 2 && InputTensorDescs.Num() <= 3);
+			check(OutputTensorDescs.Num() == 1);
+
+            const NNECore::FTensorDesc& InputA = InputTensorDescs[0];
+			const NNECore::FTensorDesc& InputB = InputTensorDescs[1];
+
+			if (InputA.GetShape().Rank() != 2)
+			{
+				UE_LOG(LogNNX, Warning, TEXT("Gemm first input should be of rank 2"));
+				return false;
+			}
+			if (InputB.GetShape().Rank() != 2)
+			{
+				UE_LOG(LogNNX, Warning, TEXT("Gemm second input should be of rank 2"));
+				return false;
+			}
+			if (InputTensorDescs.Num() == 3)
+			{
+				const NNECore::FTensorDesc& InputC = InputTensorDescs[2];
+				if (InputC.GetShape().Rank() > 2)
+				{
+					UE_LOG(LogNNX, Warning, TEXT("Gemm third input should be of rank 2 or less"));
+					return false;
+				}
+				if (InputC.GetShape().Rank() == 1 && InputC.GetShape().GetData()[0] == 1)
+				{
+					UE_LOG(LogNNX, Warning, TEXT("Gemm third input as scalar not supported"));
+					return false;
+				}
+			}
+
+			// C is treated as a scalar if there is no valid C, either width or height is zero or C dimension is 1x1
+			bIsCScalar = false; // InputTensors.Num() != 3 || InputC.Sizes[0] * InputC.Sizes[1] < 2;
+			// CScalar = C != nullptr ? C[0] : (InElementType)0;
+			bNoBias = InputTensorDescs.Num() != 3 /*|| InputC.Sizes[0] * InputC.Sizes[1] < 1*/;
+
+			InputAlpha = Attributes.GetValueOrDefault(TEXT("alpha"), InputAlpha);
+			InputBeta = Attributes.GetValueOrDefault(TEXT("beta"), InputBeta);
+			InputTransA = Attributes.GetValueOrDefault(TEXT("transA"), InputTransA);
+			InputTransB = Attributes.GetValueOrDefault(TEXT("transB"), InputTransB);
+
+			return true;
+		}
+
+		virtual void Dispatch(FRDGBuilder& GraphBuilder, TConstArrayView<FTensorRDGRef> InputTensors, TConstArrayView<FTensorRDGRef> OutputTensors) override
+		{
+			using namespace UE::NNEHlslShaders::Internal;
+
+			const EGemmAlgorithm Algorithm = EGemmAlgorithm::Simple32x32;
+
+			const float CScalar = 0.0f;
+
+			check(InputTensors.Num() >= 2 && InputTensors.Num() <= 3);
+			check(OutputTensors.Num() == 1);
+			check(InputTensors[0] != nullptr);
+			check(InputTensors[1] != nullptr);
+			check(OutputTensors[0] != nullptr);
+			const FTensorRDG& InputA = *InputTensors[0];
+			const FTensorRDG& InputB = *InputTensors[1];
+			const FTensorRDG& Output = *OutputTensors[0];
+			const FTensorRDG* InputC = nullptr;
+
+			if (InputTensors.Num() == 3)
+			{
+				check(InputTensors[2] != nullptr);
+				InputC = InputTensors[2];
+			}
+			
+			// Set parameters
+			TGemmCS::FParameters* Parameters = GraphBuilder.AllocParameters<TGemmCS::FParameters>();
+			TGemmCS::FillInParameters(InputAlpha, InputBeta, InputTransA, InputTransB, InputA, InputB, InputC, CScalar, *Parameters);
+			Parameters->A = GraphBuilder.CreateSRV(FRDGBufferSRVDesc(InputA.GetBuffer(), PF_R32_FLOAT));
+			Parameters->B = GraphBuilder.CreateSRV(FRDGBufferSRVDesc(InputB.GetBuffer(), PF_R32_FLOAT));
+			if (InputC != nullptr) {
+				Parameters->C = GraphBuilder.CreateSRV(FRDGBufferSRVDesc(InputC->GetBuffer(), PF_R32_FLOAT));
+			}
+			Parameters->Y = GraphBuilder.CreateUAV(FRDGBufferUAVDesc(Output.GetBuffer(), PF_R32_FLOAT));
+
+			TGemmCS::FPermutationDomain PermutationVector;
+			PermutationVector.Set<TGemmCS::FGemmCScalar>(bNoBias ? EGemmCScalar::NoBias : (bIsCScalar ? EGemmCScalar::Yes : EGemmCScalar::No));
+			PermutationVector.Set<TGemmCS::FGemmAlgorithm>(Algorithm);
+			PermutationVector.Set<TGemmCS::FGemmNumStackDimensions>(0);
+			TShaderMapRef<TGemmCS> ComputeShader(GetGlobalShaderMap(GMaxRHIFeatureLevel), PermutationVector);
+
+			FIntVector ThreadGroupCount = TGemmCS::GetGroupCount(*Parameters, Algorithm, 0);
+
+			RDG_EVENT_SCOPE(GraphBuilder, "NNE.Operator.Hlsl.Gemm");
+			RDG_GPU_STAT_SCOPE(GraphBuilder, FNNEOperatorGemm);
+
+			FComputeShaderUtils::AddPass(
+				GraphBuilder,
+				RDG_EVENT_NAME("NNE.Operator.Hlsl.Dispatch"),
+				ERDGPassFlags::Compute | ERDGPassFlags::NeverCull,
+				ComputeShader,
+				Parameters,
+				ThreadGroupCount);
+		}
+	};
+
+	bool ValidateGemmOperator(const NNECore::FAttributeMap& AttributeMap, TConstArrayView<ENNETensorDataType> InputTypes, TConstArrayView<NNECore::FSymbolicTensorShape> InputShapes)
+	{
+		bool bIsValid = true;
+
+		FAttributeValidator AttributeValidator;
+		AttributeValidator.AddOptional(TEXT("alpha"), ENNEAttributeDataType::Float);
+		AttributeValidator.AddOptional(TEXT("beta"), ENNEAttributeDataType::Float);
+		AttributeValidator.AddOptional(TEXT("transA"), ENNEAttributeDataType::Int32);
+		AttributeValidator.AddOptional(TEXT("transB"), ENNEAttributeDataType::Int32);
+		bIsValid &= AttributeValidator.Validate(AttributeMap);
+
+		FInputValidator InputValidator;
+		InputValidator.AddSupportedType(ENNETensorDataType::Float);
+		InputValidator.AddRequired();
+		InputValidator.AddRequired();
+		InputValidator.AddOptional();
+		bIsValid &= InputValidator.Validate(InputTypes);
+
+		return bIsValid;
+	}
+
+	FOperatorHlsl* CreateGemmOperator()
+	{
+		return new TGemm();
+	}
+
+	bool RegisterGemmOperator(FOperatorRegistryHlsl& Registry)
+	{
+		Registry.OpAdd(TEXT("Gemm"), CreateGemmOperator, ValidateGemmOperator);
+		return true;
+	}
+} // UE::NNERuntimeRDG::Private::Hlsl
diff --git a/Engine/Plugins/Experimental/NNE/Source/NNERuntimeRDG/Private/Hlsl/NNERuntimeRDGGemm.h b/Engine/Plugins/Experimental/NNE/Source/NNERuntimeRDG/Private/Hlsl/NNERuntimeRDGGemm.h
new file mode 100644
--- /dev/null
+++ b/Engine/Plugins/Experimental/NNE/Source/NNERuntimeRDG/Private/Hlsl/NNERuntimeRDGGemm.h
@@ -0,0 +1,11 @@
+// Copyright Epic Games, Inc. All Rights Reserved.
+
+#pragma once
+
+#include "NNERuntimeRDG.h"
+#include "NNERuntimeRDGHlslOp.h"
+
+namespace UE::NNERuntimeRDG::Private::Hlsl
+{
+	bool RegisterGemmOperator(FOperatorRegistryHlsl& Registry);
+} // UE::NNERuntimeRDG::Private::Hlsl
diff --git a/Engine/Plugins/Experimental/NNE/Source/NNERuntimeRDG/Private/Hlsl/NNERuntimeRDGInstanceNormalization.cpp b/Engine/Plugins/Experimental/NNE/Source/NNERuntimeRDG/Private/Hlsl/NNERuntimeRDGInstanceNormalization.cpp
new file mode 100644
--- /dev/null
+++ b/Engine/Plugins/Experimental/NNE/Source/NNERuntimeRDG/Private/Hlsl/NNERuntimeRDGInstanceNormalization.cpp
@@ -0,0 +1,178 @@
+// Copyright Epic Games, Inc. All Rights Reserved.
+
+#include "NNERuntimeRDGInstanceNormalization.h"
+#include "NNEHlslShadersInstanceNormalizationCS.h"
+#include "NNERuntimeRDGHlslHelper.h"
+#include "NNECoreTensor.h"
+#include "NNECoreTypes.h"
+
+namespace
+{
+
+int ValidateInput(const UE::NNECore::FTensorShape& Input, const UE::NNECore::FTensorShape& Scale, const UE::NNECore::FTensorShape& Bias)
+{
+	if (Input.Rank() < 3)
+	{
+		UE_LOG(LogNNX, Warning, TEXT("InstanceNormalization input data should be at least of rank 4: %d"), Input.Rank());
+		return -1;
+	}
+
+	if (Scale.Rank() != 1)
+	{
+		UE_LOG(LogNNX, Warning, TEXT("InstanceNormalization input scale should be of rank 1: %d"), Scale.Rank());
+		return -1;
+	}
+	if (Scale.GetData()[0] != Input.GetData()[1])
+	{
+		UE_LOG(LogNNX, Warning, TEXT("InstanceNormalization input scale size should be equal to channel count: %d vs %d"), Scale.GetData()[0], Input.GetData()[1]);
+		return -1;
+	}
+
+	if (Bias.GetData()[0] != Input.GetData()[1])
+	{
+		UE_LOG(LogNNX, Warning, TEXT("InstanceNormalization intput B size should be equal to channel count: : %d vs %d"), Bias.GetData()[0], Input.GetData()[1]);
+		return -1;
+	}
+	if (Bias.Rank() != 1)
+	{
+		UE_LOG(LogNNX, Warning, TEXT("InstanceNormalization input B should be of rank 1: %d"), Bias.Rank());
+		return -1;
+	}
+
+	return 0;
+}
+
+}
+
+namespace UE::NNERuntimeRDG::Private::Hlsl
+{
+	DECLARE_GPU_STAT_NAMED(FNNEOperatorInstanceNormalization, TEXT("NNE.Operator.Hlsl.InstanceNormalization"));
+
+	using EInstanceNormalizationAlgorithm = UE::NNEHlslShaders::Internal::EInstanceNormalizationAlgorithm;
+
+	/**
+	 * InstanceNormalization operator implementation
+	 */
+	class TInstanceNormalization : public FOperatorHlsl
+	{
+	public:
+
+		TInstanceNormalization() {}
+		virtual ~TInstanceNormalization() = default;
+
+	private:
+
+		float Epsilon = 1e-5;
+
+		EInstanceNormalizationAlgorithm Algorithm = EInstanceNormalizationAlgorithm::MAX;
+
+	public:
+
+		virtual int PrepareOutputs(TConstArrayView<NNECore::Internal::FTensorRef> InputTensors, TArrayView<NNECore::Internal::FTensorRef> OutputTensors) const override
+		{
+			check(InputTensors.Num() == 3);
+			check(OutputTensors.Num() == 1);
+
+			if (const int Res = ValidateInput(InputTensors[0]->GetShape(), InputTensors[1]->GetShape(), InputTensors[2]->GetShape()); Res < 0)
+			{
+				return Res;
+			}
+
+			OutputTensors[0]->SetShape(InputTensors[0]->GetShape());
+
+			return 0;
+		};
+
+		virtual bool Initialize(TConstArrayView<NNECore::FTensorDesc> InputTensorDescs, TConstArrayView<NNECore::FTensorDesc> OutputTensorDescs, const NNECore::FAttributeMap& Attributes) override
+		{
+			using namespace UE::NNEHlslShaders::Internal;
+
+			check(InputTensorDescs.Num() == 3);
+			check(OutputTensorDescs.Num() == 1);
+
+			Epsilon = Attributes.GetValue<float>(TEXT("epsilon"));
+
+			// For testing only
+			TInstanceNormalizationCS::LexFromString(Algorithm, *Attributes.GetValueOrDefault<FString>(TEXT("__UE__algorithm"), "MAX"));
+
+			return true;
+		}
+
+		virtual void Dispatch(FRDGBuilder& GraphBuilder, TConstArrayView<FTensorRDGRef> InputTensors, TConstArrayView<FTensorRDGRef> OutputTensors) override
+		{
+			using namespace UE::NNEHlslShaders::Internal;
+
+			check(InputTensors.Num() == 3);
+			check(OutputTensors.Num() == 1);
+			check(InputTensors[0] != nullptr);
+			check(InputTensors[1] != nullptr);
+			check(InputTensors[2] != nullptr);
+			check(OutputTensors[0] != nullptr);
+
+			const FTensorRDG& Input = *InputTensors[0];
+			const FTensorRDG& Scale = *InputTensors[1];
+			const FTensorRDG& Bias = *InputTensors[2];
+			const FTensorRDG& Output = *OutputTensors[0];
+
+			// Set parameters
+			TInstanceNormalizationCS::FParameters* Parameters = GraphBuilder.AllocParameters<TInstanceNormalizationCS::FParameters>();
+			TInstanceNormalizationCS::FillInParameters(Epsilon, Input, *Parameters);
+			Parameters->Input = GraphBuilder.CreateSRV(FRDGBufferSRVDesc(Input.GetBuffer(), PF_R32_FLOAT));
+			Parameters->Scale = GraphBuilder.CreateSRV(FRDGBufferSRVDesc(Scale.GetBuffer(), PF_R32_FLOAT));
+			Parameters->Bias = GraphBuilder.CreateSRV(FRDGBufferSRVDesc(Bias.GetBuffer(), PF_R32_FLOAT));
+			Parameters->Output = GraphBuilder.CreateUAV(FRDGBufferUAVDesc(Output.GetBuffer(), PF_R32_FLOAT));
+
+			EInstanceNormalizationAlgorithm DispatchAlgorithm = Algorithm;
+			if (DispatchAlgorithm == EInstanceNormalizationAlgorithm::MAX)
+			{
+				DispatchAlgorithm = TInstanceNormalizationCS::GetAlgorithm(*Parameters);
+			}
+			
+			TInstanceNormalizationCS::FPermutationDomain PermutationVector;
+			PermutationVector.Set<TInstanceNormalizationCS::FInstanceNormalizationAlgorithm>(DispatchAlgorithm);
+			TShaderMapRef<TInstanceNormalizationCS> ComputeShader(GetGlobalShaderMap(GMaxRHIFeatureLevel), PermutationVector);
+			FIntVector ThreadGroupCount = TInstanceNormalizationCS::GetGroupCount(*Parameters, DispatchAlgorithm);
+
+			RDG_EVENT_SCOPE(GraphBuilder, "NNE.Operator.Hlsl.InstanceNormalization");
+			RDG_GPU_STAT_SCOPE(GraphBuilder, FNNEOperatorInstanceNormalization);
+
+			FComputeShaderUtils::AddPass(
+				GraphBuilder,
+				RDG_EVENT_NAME("NNI.Operator.Hlsl.InstanceNormalization.Dispatch"),
+				ERDGPassFlags::Compute | ERDGPassFlags::NeverCull,
+				ComputeShader,
+				Parameters,
+				ThreadGroupCount);
+		}
+	};
+
+	bool ValidateInstanceNormalizationOperator(const NNECore::FAttributeMap& AttributeMap, TConstArrayView<ENNETensorDataType> InputTypes, TConstArrayView<NNECore::FSymbolicTensorShape> InputShapes)
+	{
+		bool bIsValid = true;
+
+		FAttributeValidator AttributeValidator;
+		AttributeValidator.AddRequired(TEXT("epsilon"), ENNEAttributeDataType::Float);
+		AttributeValidator.AddOptional(TEXT("__UE__algorithm"), ENNEAttributeDataType::String);
+		bIsValid &= AttributeValidator.Validate(AttributeMap);
+
+		FInputValidator InputValidator;
+		InputValidator.AddSupportedType(ENNETensorDataType::Float);
+		InputValidator.AddRequired();
+		InputValidator.AddRequired();
+		InputValidator.AddRequired();
+		bIsValid &= InputValidator.Validate(InputTypes);
+
+		return bIsValid;
+	}
+
+	FOperatorHlsl* CreateInstanceNormalizationOperator()
+	{
+		return new TInstanceNormalization();
+	}
+
+	bool RegisterInstanceNormalizationOperator(FOperatorRegistryHlsl& Registry)
+	{
+		Registry.OpAdd(TEXT("InstanceNormalization"), CreateInstanceNormalizationOperator, ValidateInstanceNormalizationOperator);
+		return true;
+	}
+} // UE::NNERuntimeRDG::Private::Hlsl
diff --git a/Engine/Plugins/Experimental/NNE/Source/NNERuntimeRDG/Private/Hlsl/NNERuntimeRDGInstanceNormalization.h b/Engine/Plugins/Experimental/NNE/Source/NNERuntimeRDG/Private/Hlsl/NNERuntimeRDGInstanceNormalization.h
new file mode 100644
--- /dev/null
+++ b/Engine/Plugins/Experimental/NNE/Source/NNERuntimeRDG/Private/Hlsl/NNERuntimeRDGInstanceNormalization.h
@@ -0,0 +1,11 @@
+// Copyright Epic Games, Inc. All Rights Reserved.
+
+#pragma once
+
+#include "NNERuntimeRDG.h"
+#include "NNERuntimeRDGHlslOp.h"
+
+namespace UE::NNERuntimeRDG::Private::Hlsl
+{
+	bool RegisterInstanceNormalizationOperator(FOperatorRegistryHlsl& Registry);
+} // UE::NNERuntimeRDG::Private::Hlsl
diff --git a/Engine/Plugins/Experimental/NNE/Source/NNERuntimeRDG/Private/Hlsl/NNERuntimeRDGMatMul.cpp b/Engine/Plugins/Experimental/NNE/Source/NNERuntimeRDG/Private/Hlsl/NNERuntimeRDGMatMul.cpp
new file mode 100644
--- /dev/null
+++ b/Engine/Plugins/Experimental/NNE/Source/NNERuntimeRDG/Private/Hlsl/NNERuntimeRDGMatMul.cpp
@@ -0,0 +1,168 @@
+// Copyright Epic Games, Inc. All Rights Reserved.
+
+#include "NNERuntimeRDGMatMul.h"
+#include "NNEHlslShadersGemmCS.h"
+#include "NNERuntimeRDGHlslHelper.h"
+#include "NNECoreTensor.h"
+#include "NNECoreTypes.h"
+
+namespace UE::NNERuntimeRDG::Private::Hlsl
+{
+    DECLARE_GPU_STAT_NAMED(FNNEOperatorMatMul, TEXT("NNE.Operator.Hlsl.MatMul"));
+
+	/**
+	 * MatMul operator implementation
+	 */
+	class TMatMul : public FOperatorHlsl
+	{
+	public:
+
+		TMatMul() {}
+		virtual ~TMatMul() = default;
+
+	public:
+
+		virtual int PrepareOutputs(TConstArrayView<NNECore::Internal::FTensorRef> InputTensors, TArrayView<NNECore::Internal::FTensorRef> OutputTensors) const override
+		{
+			check(InputTensors.Num() == 2);
+			check(OutputTensors.Num() == 1);
+
+			const NNX::FTensorShape& InputA = InputTensors[0]->GetShape();
+			const NNX::FTensorShape& InputB = InputTensors[1]->GetShape();
+
+			if (InputA.Rank() < 2)
+			{
+				UE_LOG(LogNNX, Warning, TEXT("Matmul first input should be at least of rank 2"));
+				return -1;
+			}
+			if (InputB.Rank() < 2)
+			{
+				UE_LOG(LogNNX, Warning, TEXT("Matmul second input should be at least of rank 2"));
+				return -1;
+			}
+			if (InputA.GetData()[InputA.Rank() - 1] != InputB.GetData()[InputB.Rank() - 2])
+			{
+				UE_LOG(LogNNX, Warning, TEXT("Matmul first input last dimension should be equal to second input last dimension"));
+				return -1;
+			}
+
+			const int32 OutputRank = FMath::Max(InputA.Rank(), InputB.Rank());
+			TArray<uint32> OutputShapeData;
+			OutputShapeData.SetNumUninitialized(OutputRank);
+			
+			//Broadcast
+			for (int32 i = 0; i < OutputRank; ++i)
+			{
+				int32 AIndex = InputA.Rank() - 1 - i;
+				int32 BIndex = InputB.Rank() - 1 - i;
+				int32 AValue = AIndex >= 0 ? InputA.GetData()[AIndex] : 1;
+				int32 BValue = BIndex >= 0 ? InputB.GetData()[BIndex] : 1;
+				int32 OutputValue = FMath::Max(AValue, BValue);
+				OutputShapeData[OutputRank - 1 - i] = OutputValue;
+			}
+			
+			//2D Mat
+			OutputShapeData[OutputRank - 2] = InputA.GetData()[InputA.Rank() - 2];
+			OutputShapeData[OutputRank - 1] = InputB.GetData()[InputB.Rank() - 1];
+
+			NNECore::FTensorShape OutputShape = NNECore::FTensorShape::Make(OutputShapeData);
+
+			OutputTensors[0]->SetShape(OutputShape);
+
+			return 0;
+		};
+
+		virtual bool Initialize(TConstArrayView<NNECore::FTensorDesc> InputTensorDescs, TConstArrayView<NNECore::FTensorDesc> OutputTensorDescs, const NNECore::FAttributeMap& Attributes) override
+		{
+			check(InputTensorDescs.Num() == 2);
+			check(OutputTensorDescs.Num() == 1);
+
+			const NNECore::FTensorDesc& InputA = InputTensorDescs[0];
+			const NNECore::FTensorDesc& InputB = InputTensorDescs[1];
+
+			if (InputA.GetShape().Rank() < 2)
+			{
+				UE_LOG(LogNNX, Warning, TEXT("Matmul first input should be at least of rank 2"));
+				return false;
+			}
+			if (InputB.GetShape().Rank() < 2)
+			{
+				UE_LOG(LogNNX, Warning, TEXT("Matmul second input should be at least of rank 2"));
+				return false;
+			}
+
+			return true;
+		}
+
+		virtual void Dispatch(FRDGBuilder& GraphBuilder, TConstArrayView<FTensorRDGRef> InputTensors, TConstArrayView<FTensorRDGRef> OutputTensors) override
+		{
+			using namespace UE::NNEHlslShaders::Internal;
+
+			check(InputTensors.Num() == 2);
+			check(OutputTensors.Num() == 1);
+			check(InputTensors[0] != nullptr);
+			check(InputTensors[1] != nullptr);
+			check(OutputTensors[0] != nullptr);
+			const FTensorRDG& InputA = *InputTensors[0];
+			const FTensorRDG& InputB = *InputTensors[1];
+			const FTensorRDG& Output = *OutputTensors[0];
+
+			const EGemmAlgorithm Algorithm = EGemmAlgorithm::Simple32x32;
+
+			const int32 NumStackDimensions = FMath::Max(FMath::Max(InputA.GetShape().Rank(), InputB.GetShape().Rank()) - 2, 0);
+
+			// Set parameters
+			TGemmCS::FParameters* Parameters = GraphBuilder.AllocParameters<TGemmCS::FParameters>();
+			TGemmCS::FillInParametersMatMul(InputA, InputB, *Parameters);
+			Parameters->A = GraphBuilder.CreateSRV(FRDGBufferSRVDesc(InputA.GetBuffer(), PF_R32_FLOAT));
+			Parameters->B = GraphBuilder.CreateSRV(FRDGBufferSRVDesc(InputB.GetBuffer(), PF_R32_FLOAT));
+			Parameters->Y = GraphBuilder.CreateUAV(FRDGBufferUAVDesc(Output.GetBuffer(), PF_R32_FLOAT));
+
+			TGemmCS::FPermutationDomain PermutationVector;
+			PermutationVector.Set<TGemmCS::FGemmCScalar>(EGemmCScalar::NoBias);
+			PermutationVector.Set<TGemmCS::FGemmAlgorithm>(Algorithm);
+			PermutationVector.Set<TGemmCS::FGemmNumStackDimensions>(NumStackDimensions);
+			TShaderMapRef<TGemmCS> ComputeShader(GetGlobalShaderMap(GMaxRHIFeatureLevel), PermutationVector);
+
+			FIntVector ThreadGroupCount = TGemmCS::GetGroupCount(*Parameters, Algorithm, NumStackDimensions);
+
+			RDG_EVENT_SCOPE(GraphBuilder, "NNE.Operator.Hlsl.MatMul");
+			RDG_GPU_STAT_SCOPE(GraphBuilder, FNNEOperatorMatMul);
+
+			FComputeShaderUtils::AddPass(
+				GraphBuilder,
+				RDG_EVENT_NAME("NNE.Operator.Hlsl.MatMul.Dispatch"),
+				ERDGPassFlags::Compute | ERDGPassFlags::NeverCull,
+				ComputeShader,
+				Parameters,
+				ThreadGroupCount);
+		}
+	};
+
+	bool ValidateMatMulOperator(const NNECore::FAttributeMap& AttributeMap, TConstArrayView<ENNETensorDataType> InputTypes, TConstArrayView<NNECore::FSymbolicTensorShape> InputShapes)
+	{
+		bool bIsValid = true;
+
+		FAttributeValidator AttributeValidator;
+		bIsValid &= AttributeValidator.Validate(AttributeMap);
+
+		FInputValidator InputValidator;
+		InputValidator.AddSupportedType(ENNETensorDataType::Float);
+		InputValidator.AddRequired();
+		InputValidator.AddRequired();
+		bIsValid &= InputValidator.Validate(InputTypes);
+
+		return bIsValid;
+	}
+
+	FOperatorHlsl* CreateMatMulOperator()
+	{
+		return new TMatMul();
+	}
+
+	bool RegisterMatMulOperator(FOperatorRegistryHlsl& Registry)
+	{
+		Registry.OpAdd(TEXT("MatMul"), CreateMatMulOperator, ValidateMatMulOperator);
+		return true;
+	}
+} // UE::NNERuntimeRDG::Private::Hlsl
diff --git a/Engine/Plugins/Experimental/NNE/Source/NNERuntimeRDG/Private/Hlsl/NNERuntimeRDGMatMul.h b/Engine/Plugins/Experimental/NNE/Source/NNERuntimeRDG/Private/Hlsl/NNERuntimeRDGMatMul.h
new file mode 100644
--- /dev/null
+++ b/Engine/Plugins/Experimental/NNE/Source/NNERuntimeRDG/Private/Hlsl/NNERuntimeRDGMatMul.h
@@ -0,0 +1,11 @@
+// Copyright Epic Games, Inc. All Rights Reserved.
+
+#pragma once
+
+#include "NNERuntimeRDG.h"
+#include "NNERuntimeRDGHlslOp.h"
+
+namespace UE::NNERuntimeRDG::Private::Hlsl
+{
+	bool RegisterMatMulOperator(FOperatorRegistryHlsl& Registry);
+} // UE::NNERuntimeRDG::Private::Hlsl
diff --git a/Engine/Plugins/Experimental/NNE/Source/NNERuntimeRDG/Private/Hlsl/NNERuntimeRDGPad.cpp b/Engine/Plugins/Experimental/NNE/Source/NNERuntimeRDG/Private/Hlsl/NNERuntimeRDGPad.cpp
new file mode 100644
--- /dev/null
+++ b/Engine/Plugins/Experimental/NNE/Source/NNERuntimeRDG/Private/Hlsl/NNERuntimeRDGPad.cpp
@@ -0,0 +1,167 @@
+// Copyright Epic Games, Inc. All Rights Reserved.
+
+#include "NNERuntimeRDGPad.h"
+#include "NNEHlslShadersPadCS.h"
+#include "NNERuntimeRDGHlslHelper.h"
+#include "NNECoreTensor.h"
+#include "NNECoreTypes.h"
+
+namespace UE::NNERuntimeRDG::Private::Hlsl
+{
+	DECLARE_GPU_STAT_NAMED(FNNEOperatorPad, TEXT("NNE.Operator.Hlsl.Pad"));
+
+	using EPadMode = UE::NNEHlslShaders::Internal::EPadMode;
+
+	/**
+	 * Pad operator implementation
+	 */
+	class FPad : public FOperatorHlsl
+	{
+	public:
+
+		FPad() {}
+		virtual ~FPad() = default;
+
+		TArray<int32> Pads;
+		float Value;
+		EPadMode Mode;
+
+	public:
+
+		virtual int PrepareOutputs(TConstArrayView<NNECore::Internal::FTensorRef> InputTensors, TArrayView<NNECore::Internal::FTensorRef> OutputTensors) const override
+		{
+			check(InputTensors.Num() == 1);
+			check(OutputTensors.Num() == 1);
+
+			const NNECore::Internal::FTensor& X = *InputTensors[0];
+
+			TArray<uint32> OutputShapeData;
+			for (int32 i = 0; i < X.GetShape().Rank(); ++i)
+			{
+				int32 PrePad = Pads[i];
+				int32 PostPad = Pads[i + X.GetShape().Rank()];
+				int32 OutputDim = PrePad + X.GetShape().GetData()[i] + PostPad;
+				if (OutputDim < 1)
+				{
+					UE_LOG(LogNNX, Warning, TEXT("Pads cannot reduce dimension below 1, but would for tensor (name:%s) at rank %d of size %d with prepad %d and postpad %d."), *X.GetName(), i, X.GetShape().GetData()[i], PrePad, PostPad);
+					return -1;
+				}
+				OutputShapeData.Emplace(OutputDim);
+			}
+
+			NNECore::FTensorShape OutputShape = NNECore::FTensorShape::Make(OutputShapeData);
+			OutputTensors[0]->SetShape(OutputShape);
+
+			return 0;
+		};
+
+		virtual bool Initialize(TConstArrayView<NNECore::FTensorDesc> InputTensorDescs, TConstArrayView<NNECore::FTensorDesc> OutputTensorDescs, const NNECore::FAttributeMap& Attributes) override
+		{
+			using namespace UE::NNEHlslShaders::Internal;
+			
+			check(InputTensorDescs.Num() == 1);
+			check(OutputTensorDescs.Num() == 1);
+			
+			Pads = Attributes.GetValue<TArray<int32>>(TEXT("pads"));
+			Value = Attributes.GetValueOrDefault<float>(TEXT("value"), 0.0f);
+			FPadCS::LexFromString(Mode, *Attributes.GetValue<FString>(TEXT("mode")));
+
+			if ((2*InputTensorDescs[0].GetShape().Rank()) != Pads.Num())
+			{
+				UE_LOG(LogNNX, Warning, TEXT("pads attribute lenght (%d) should be twice the rank of input X (%d)."), Pads.Num(), InputTensorDescs[0].GetShape().Rank());
+				return false;
+			}
+
+			return true;
+		}
+
+		virtual void Dispatch(FRDGBuilder& GraphBuilder, TConstArrayView<FTensorRDGRef> InputTensors, TConstArrayView<FTensorRDGRef> OutputTensors) override
+		{
+			using namespace UE::NNEHlslShaders::Internal;
+			
+			check(InputTensors.Num() == 1);
+			check(OutputTensors.Num() == 1);
+			check(InputTensors[0] != nullptr);
+			check(OutputTensors[0] != nullptr);
+			const FTensorRDG& Input = *InputTensors[0];
+			const FTensorRDG& Output = *OutputTensors[0];
+
+			FRDGBufferSRVRef InputSRV = GraphBuilder.CreateSRV(FRDGBufferSRVDesc(Input.GetBuffer(), PF_R32_FLOAT));
+			FRDGBufferUAVRef OutputUAV = GraphBuilder.CreateUAV(FRDGBufferUAVDesc(Output.GetBuffer(), PF_R32_FLOAT));
+
+			FIntVector ThreadGroupCount = ComputeElementWiseThreadGroups(Output.GetVolume(), FPadConstants::NUM_GROUP_THREADS);
+
+			// Set parameters
+			FPadCS::FParameters* Params = GraphBuilder.AllocParameters<FPadCS::FParameters>();
+			Params->Input = InputSRV;
+			Params->Output = OutputUAV;
+			FillTensorStrideShaderParameters(Input, Params->TensorInfo, 0);
+			FillTensorStrideShaderParameters(Output, Params->TensorInfo, 1);
+			FillTensorSizeShaderParameters(Input, Params->TensorInfo, 2);
+			for (int32 i = 0; i < Input.GetShape().Rank(); ++i)
+			{
+				Params->TensorInfo[i][3] = Pads[i];//Pre-pad
+			}
+			Params->Value = Value;
+			Params->Num = Output.GetVolume();
+			Params->ThreadCountX = ThreadGroupCount.X * FPadConstants::NUM_GROUP_THREADS;
+
+			FPadCS::FPermutationDomain PermutationVector;
+			PermutationVector.Set<FPadCS::FPadMode>(Mode);
+			PermutationVector.Set<FPadCS::FPadNumDimensions>(Output.GetShape().Rank());
+
+			TShaderMapRef<FPadCS> ComputeShader(GetGlobalShaderMap(GMaxRHIFeatureLevel), PermutationVector);
+
+			RDG_EVENT_SCOPE(GraphBuilder, "NNE.Operator.Hlsl.Pad");
+			RDG_GPU_STAT_SCOPE(GraphBuilder, FNNEOperatorPad);
+
+			FComputeShaderUtils::AddPass(
+				GraphBuilder,
+				RDG_EVENT_NAME("NNE.Operator.Hlsl.Pad.Dispatch"),
+				ERDGPassFlags::Compute | ERDGPassFlags::NeverCull,
+				ComputeShader,
+				Params,
+				ThreadGroupCount);
+		}
+	};
+
+	bool ValidatePadOperator(const NNECore::FAttributeMap& AttributeMap, TConstArrayView<ENNETensorDataType> InputTypes, TConstArrayView<NNECore::FSymbolicTensorShape> InputShapes)
+	{
+		bool bIsValid = true;
+
+		//This match version 2 of the pad operator (next version is with opset 11)
+		//[REDACTED_PROFILE_URL]/onnx/blob/main/docs/Changelog.md#Pad-2
+		FAttributeValidator AttributeValidator;
+		AttributeValidator.AddOptional(TEXT("mode"), ENNEAttributeDataType::String);
+		AttributeValidator.AddRequired(TEXT("pads"), ENNEAttributeDataType::Int32Array);
+		AttributeValidator.AddOptional(TEXT("value"), ENNEAttributeDataType::Float);
+		bIsValid &= AttributeValidator.Validate(AttributeMap);
+
+		for (int32 Pad : AttributeMap.GetValue<TArray<int32>>(TEXT("pads")))
+		{
+			if (Pad < 0)
+			{
+				UE_LOG(LogNNX, Warning, TEXT("Pad operator does not support negative padding at the moment."));
+				return false;
+			}
+		}
+		
+		FInputValidator InputValidator;
+		InputValidator.AddSupportedType(ENNETensorDataType::Float);
+		InputValidator.AddRequired();
+		bIsValid &= InputValidator.Validate(InputTypes);
+
+		return bIsValid;
+	}
+
+	FOperatorHlsl* CreatePadOperator()
+	{
+		return new FPad();
+	}
+
+	bool RegisterPadOperator(FOperatorRegistryHlsl& Registry)
+	{
+		Registry.OpAdd(TEXT("Pad"), CreatePadOperator, ValidatePadOperator);
+		return true;
+	}
+} // UE::NNERuntimeRDG::Private::Hlsl
diff --git a/Engine/Plugins/Experimental/NNE/Source/NNERuntimeRDG/Private/Hlsl/NNERuntimeRDGPad.h b/Engine/Plugins/Experimental/NNE/Source/NNERuntimeRDG/Private/Hlsl/NNERuntimeRDGPad.h
new file mode 100644
--- /dev/null
+++ b/Engine/Plugins/Experimental/NNE/Source/NNERuntimeRDG/Private/Hlsl/NNERuntimeRDGPad.h
@@ -0,0 +1,11 @@
+// Copyright Epic Games, Inc. All Rights Reserved.
+
+#pragma once
+
+#include "NNERuntimeRDG.h"
+#include "NNERuntimeRDGHlslOp.h"
+
+namespace UE::NNERuntimeRDG::Private::Hlsl
+{
+	bool RegisterPadOperator(FOperatorRegistryHlsl& Registry);
+} // UE::NNERuntimeRDG::Private::Hlsl
diff --git a/Engine/Plugins/Experimental/NNE/Source/NNERuntimeRDG/Private/Hlsl/NNERuntimeRDGUpsample.cpp b/Engine/Plugins/Experimental/NNE/Source/NNERuntimeRDG/Private/Hlsl/NNERuntimeRDGUpsample.cpp
new file mode 100644
--- /dev/null
+++ b/Engine/Plugins/Experimental/NNE/Source/NNERuntimeRDG/Private/Hlsl/NNERuntimeRDGUpsample.cpp
@@ -0,0 +1,153 @@
+// Copyright Epic Games, Inc. All Rights Reserved.
+
+#include "NNERuntimeRDGUpsample.h"
+#include "NNEHlslShadersUpsampleCS.h"
+#include "NNERuntimeRDGHlslHelper.h"
+#include "NNECoreTypes.h"
+#include "NNECoreTensor.h"
+
+namespace UE::NNERuntimeRDG::Private::Hlsl
+{
+	DECLARE_GPU_STAT_NAMED(FNNEOperatorUpsample, TEXT("NNE.Operator.Hlsl.Upsample"));
+
+	/**
+	 * Upsample operator implementation
+	 */
+	class FUpsample : public FOperatorHlsl
+	{
+	public:
+
+		FUpsample() {}
+		virtual ~FUpsample() = default;
+
+	public:
+
+		virtual int PrepareOutputs(TConstArrayView<NNECore::Internal::FTensorRef> InputTensors, TArrayView<NNECore::Internal::FTensorRef> OutputTensors) const override
+		{
+			check(InputTensors.Num() == 2);
+			check(OutputTensors.Num() == 1);
+
+			const NNECore::Internal::FTensor& X = *InputTensors[0];
+			const NNECore::Internal::FTensor& Scales = *InputTensors[1];
+			
+			if (!Scales.HasPreparedData())
+			{
+				UE_LOG(LogNNX, Warning, TEXT("Upsample input 'Scale' (name: %s) should be constant for shape inference to succeed, however it is not."), *Scales.GetName());
+				return -1;
+			}
+
+			TConstArrayView<float> ScalesData = Scales.GetPreparedData<float>();
+
+			if (ScalesData.Num() != X.GetShape().Rank())
+			{
+				UE_LOG(LogNNX, Warning, TEXT("Upsample input 'Scale' (name: %s) have %d elements. While it should be the same as the rank of input 'X' (name : %s) witch is %d"), *Scales.GetName(), ScalesData.Num(), *X.GetName(), X.GetShape().Rank());
+				return -1;
+			}
+
+			TArray<uint32> OutputShapeData;
+			for (int32 i = 0; i < X.GetShape().Rank(); ++i)
+			{
+				
+				OutputShapeData.Emplace(FMath::FloorToInt32(X.GetShape().GetData()[i] * ScalesData[i]));
+			}
+
+			NNECore::FTensorShape OutputShape = NNECore::FTensorShape::Make(OutputShapeData);
+			OutputTensors[0]->SetShape(OutputShape);
+
+			return 0;
+		};
+
+		virtual bool Initialize(TConstArrayView<NNX::FTensorDesc> InputTensorDescs, TConstArrayView<NNX::FTensorDesc> OutputTensorDescs, const NNECore::FAttributeMap& Attributes) override
+		{
+			check(InputTensorDescs.Num() == 2);
+			check(OutputTensorDescs.Num() == 1);
+
+			return true;
+		}
+
+		virtual void Dispatch(FRDGBuilder& GraphBuilder, TConstArrayView<FTensorRDGRef> InputTensors, TConstArrayView<FTensorRDGRef> OutputTensors) override
+		{
+			using namespace UE::NNEHlslShaders::Internal;
+			
+			check(InputTensors.Num() == 2);
+			check(OutputTensors.Num() == 1);
+			check(InputTensors[0] != nullptr);
+			check(InputTensors[1] != nullptr);
+			check(OutputTensors[0] != nullptr);
+			const FTensorRDG& Input = *InputTensors[0];
+			const FTensorRDG& Scales = *InputTensors[1];
+			const FTensorRDG& Output = *OutputTensors[0];
+
+			check(Scales.HasPreparedData());
+			TConstArrayView<float> ScalesData = Scales.GetPreparedData<float>();
+
+			FRDGBufferSRVRef InputSRV = GraphBuilder.CreateSRV(FRDGBufferSRVDesc(Input.GetBuffer(), PF_R32_FLOAT));
+			FRDGBufferUAVRef OutputUAV = GraphBuilder.CreateUAV(FRDGBufferUAVDesc(Output.GetBuffer(), PF_R32_FLOAT));
+
+			FIntVector ThreadGroupCount = ComputeElementWiseThreadGroups(Output.GetVolume(), FUpsampleConstants::NUM_GROUP_THREADS);
+
+			// Set parameters
+			FUpsampleCS::FParameters* Params = GraphBuilder.AllocParameters<FUpsampleCS::FParameters>();
+			Params->Input = InputSRV;
+			Params->Output = OutputUAV;
+			FillTensorStrideShaderParameters(Input, Params->TensorInfo, 0);
+			FillTensorStrideShaderParameters(Output, Params->TensorInfo, 1);
+			FillTensorSizeShaderParameters(Input, Params->TensorInfo, 2);
+			FillTensorSizeShaderParameters(Output, Params->TensorInfo, 3);
+			Params->Num = Output.GetVolume();
+			Params->ThreadCountX = ThreadGroupCount.X * FUpsampleConstants::NUM_GROUP_THREADS;
+
+			FUpsampleCS::FPermutationDomain PermutationVector;
+
+			PermutationVector.Set<FUpsampleCS::FUpsampleNumDimensions>(Output.GetShape().Rank());
+
+			TShaderMapRef<FUpsampleCS> ComputeShader(GetGlobalShaderMap(GMaxRHIFeatureLevel), PermutationVector);
+
+			RDG_EVENT_SCOPE(GraphBuilder, "NNE.Operator.Hlsl.Upsample");
+			RDG_GPU_STAT_SCOPE(GraphBuilder, FNNEOperatorUpsample);
+
+			FComputeShaderUtils::AddPass(
+				GraphBuilder,
+				RDG_EVENT_NAME("NNE.Operator.Hlsl.Upsample.Dispatch"),
+				ERDGPassFlags::Compute | ERDGPassFlags::NeverCull,
+				ComputeShader,
+				Params,
+				ThreadGroupCount);
+		}
+	};
+
+	bool ValidateUpsampleOperator(const NNECore::FAttributeMap& AttributeMap, TConstArrayView<ENNETensorDataType> InputTypes, TConstArrayView<NNECore::FSymbolicTensorShape> InputShapes)
+	{
+		bool bIsValid = true;
+
+		FAttributeValidator AttributeValidator;
+		AttributeValidator.AddOptional(TEXT("mode"), ENNEAttributeDataType::String);
+		bIsValid &= AttributeValidator.Validate(AttributeMap);
+		
+		FString Mode = AttributeMap.GetValueOrDefault<FString>(TEXT("mode"), TEXT("nearest"));
+		if (!Mode.Equals(TEXT("nearest")))
+		{
+			UE_LOG(LogNNX, Warning, TEXT("Upsample HLSL operator only supports nearest mode for now"));
+			return false;
+		}
+
+		FInputValidator InputValidator;
+		InputValidator.AddSupportedType(ENNETensorDataType::Float);
+		InputValidator.AddRequired();
+		InputValidator.AddRequired();
+		bIsValid &= InputValidator.Validate(InputTypes);
+
+		return bIsValid;
+	}
+
+	FOperatorHlsl* CreateUpsampleOperator()
+	{
+		return new FUpsample();
+	}
+
+	bool RegisterUpsampleOperator(FOperatorRegistryHlsl& Registry)
+	{
+		Registry.OpAdd(TEXT("Upsample"), CreateUpsampleOperator, ValidateUpsampleOperator);
+		return true;
+	}
+} // UE::NNERuntimeRDG::Private::Hlsl
diff --git a/Engine/Plugins/Experimental/NNE/Source/NNERuntimeRDG/Private/Hlsl/NNERuntimeRDGUpsample.h b/Engine/Plugins/Experimental/NNE/Source/NNERuntimeRDG/Private/Hlsl/NNERuntimeRDGUpsample.h
new file mode 100644
--- /dev/null
+++ b/Engine/Plugins/Experimental/NNE/Source/NNERuntimeRDG/Private/Hlsl/NNERuntimeRDGUpsample.h
@@ -0,0 +1,11 @@
+// Copyright Epic Games, Inc. All Rights Reserved.
+
+#pragma once
+
+#include "NNERuntimeRDG.h"
+#include "NNERuntimeRDGHlslOp.h"
+
+namespace UE::NNERuntimeRDG::Private::Hlsl
+{
+	bool RegisterUpsampleOperator(FOperatorRegistryHlsl& Registry);
+} // UE::NNERuntimeRDG::Private::Hlsl
diff --git a/Engine/Plugins/Experimental/NNE/Source/NNERuntimeRDG/Private/NNERuntimeRDG.cpp b/Engine/Plugins/Experimental/NNE/Source/NNERuntimeRDG/Private/NNERuntimeRDG.cpp
new file mode 100644
--- /dev/null
+++ b/Engine/Plugins/Experimental/NNE/Source/NNERuntimeRDG/Private/NNERuntimeRDG.cpp
@@ -0,0 +1,643 @@
+// Copyright Epic Games, Inc. All Rights Reserved.
+
+#include "NNERuntimeRDG.h"
+#include "NNXInferenceModel.h"
+#include "NNXRuntimeFormat.h"
+#include "NNEUtilsModelOptimizer.h"
+
+#include "RenderGraphBuilder.h"
+#include "RenderGraphUtils.h"
+#include "RHIGPUReadback.h"
+
+#include "Serialization/MemoryReader.h"
+
+namespace UE::NNERuntimeRDG::Private
+{
+
+FGuid FRuntimeRDG::GUID = FGuid((int32)'R', (int32)'R', (int32)'D', (int32)'G');
+int32 FRuntimeRDG::Version = 0x00000001;
+
+struct FInferenceModelRDG::FReadbackEntry
+{
+	TUniquePtr<FRHIGPUBufferReadback>	RHI;
+	void*					CpuMemory;
+	size_t					Offset;
+	size_t					Size;
+};
+
+//
+//
+//
+bool AlwaysValidValidationFunction(
+	const NNECore::FAttributeMap& AttributeMap, 
+	TConstArrayView<ENNETensorDataType> InputTensorTypes,
+	TConstArrayView<NNECore::FSymbolicTensorShape> InputShapes)
+{
+	return true;
+}
+
+
+//
+//
+//
+FInputValidator::FInputValidator() : 
+	NumRequiredInput(0), NumOptionalInput(0)
+{
+	TemplateTypes.SetNum(1);
+}
+
+bool FInputValidator::Validate(TConstArrayView<ENNETensorDataType> InputTypes)
+{
+	check(InputTemplateIndices.Num() == NumRequiredInput + NumOptionalInput);
+
+	bool bAreInputValid = true;
+	int32 NumInputsToValidate = FMath::Min(InputTemplateIndices.Num(), InputTypes.Num());
+	
+	if (InputTypes.Num() < NumRequiredInput)
+	{
+		UE_LOG(LogNNX, Warning, TEXT("Required '%d' inputs but found '%d'."), NumRequiredInput, InputTypes.Num());
+		bAreInputValid = false;
+	}
+	if (InputTypes.Num() > NumRequiredInput+NumOptionalInput)
+	{
+		UE_LOG(LogNNX, Warning, TEXT("Got a total of '%d' inputs but should have '%d' maximum."), InputTypes.Num(), NumRequiredInput + NumOptionalInput);
+		bAreInputValid = false;
+	}
+	
+	for (int32 Idx = 0; Idx < NumInputsToValidate; ++Idx)
+	{
+		const int32 TemplateIdx = InputTemplateIndices[Idx];
+		
+		check(TemplateIdx < TemplateTypes.Num());
+		if (INDEX_NONE == TemplateTypes[TemplateIdx].Find(InputTypes[Idx]))
+		{
+			UE_LOG(LogNNX, Warning, TEXT("Input '%d' from template idx '%d' is of type '%d' is not supported."), Idx, TemplateIdx, (int)InputTypes[Idx]);
+			bAreInputValid = false;
+		}
+	}
+	return bAreInputValid;
+}
+void FInputValidator::SetTemplateCount(int TemplateCount)
+{
+	TemplateTypes.SetNum(TemplateCount);
+}
+void FInputValidator::AddSupportedType(ENNETensorDataType Type, int TemplateIdx)
+{
+	check(TemplateTypes.Num() > TemplateIdx);
+	TemplateTypes[TemplateIdx].Add(Type);
+}
+void FInputValidator::AddOptional(int32 TemplateIdx)
+{
+	InputTemplateIndices.Add(TemplateIdx);
+	++NumOptionalInput;
+}
+
+void FInputValidator::AddRequired(int32 TemplateIdx)
+{
+	checkf(NumOptionalInput==0, TEXT("All required attribute should be declared before the optional ones as they are referenced by indices"));
+	InputTemplateIndices.Add(TemplateIdx);
+	++NumRequiredInput;
+}
+
+
+//
+//
+//
+void FAttributeValidator::AddOptional(const FString& Name, ENNEAttributeDataType Type)
+{
+	checkf(nullptr == OptionalAttributes.FindByPredicate([Name](const FEntry& Other) { return Other.Name == Name; }), TEXT("Attribute name should be unique"));
+	checkf(nullptr == RequiredAttributes.FindByPredicate([Name](const FEntry& Other) { return Other.Name == Name; }), TEXT("Attribute name should be unique"));
+	OptionalAttributes.Emplace(Name, Type);
+}
+
+void FAttributeValidator::AddRequired(const FString& Name, ENNEAttributeDataType Type)
+{
+	checkf(nullptr == OptionalAttributes.FindByPredicate([Name](const FEntry& Other) { return Other.Name == Name; }), TEXT("Attribute name should be unique"));
+	checkf(nullptr == RequiredAttributes.FindByPredicate([Name](const FEntry& Other) { return Other.Name == Name; }), TEXT("Attribute name should be unique"));
+	RequiredAttributes.Emplace(Name, Type);
+}
+
+bool FAttributeValidator::Validate(const NNECore::FAttributeMap& AttributesToValidate)
+{
+	bool bAreAttributesValid = true;
+
+	//Verify all required attribute are matching specifications
+	for (int32 Idx = 0; Idx < RequiredAttributes.Num(); ++Idx)
+	{
+		const FNNEAttributeValue* FoundAttribute = AttributesToValidate.GetAttributeValue(RequiredAttributes[Idx].Name);
+		
+		if (FoundAttribute == nullptr)
+		{
+			bAreAttributesValid = false;
+			UE_LOG(LogNNX, Warning, TEXT("Required attribute '%s' not found."),
+				*RequiredAttributes[Idx].Name);
+		}
+		else if (RequiredAttributes[Idx].Type != FoundAttribute->GetType())
+		{
+			bAreAttributesValid = false;
+			UE_LOG(LogNNX, Warning, TEXT("Required attribute '%s' type '%d' does not match expected type '%d'."),
+				*RequiredAttributes[Idx].Name,
+				(int)FoundAttribute->GetType(),
+				(int)RequiredAttributes[Idx].Type);
+		}
+	}
+
+	//Verify all optional attribute are matching specifications
+	for (int32 Idx = 0; Idx < OptionalAttributes.Num(); ++Idx)
+	{
+		const FNNEAttributeValue* FoundAttribute = AttributesToValidate.GetAttributeValue(OptionalAttributes[Idx].Name);
+		
+		if ((FoundAttribute != nullptr) && (OptionalAttributes[Idx].Type != FoundAttribute->GetType()))
+		{
+			bAreAttributesValid = false;
+			UE_LOG(LogNNX, Warning, TEXT("Optional attribute '%s' type '%d' does not match expected type '%d'."),
+				*OptionalAttributes[Idx].Name,
+				(int)FoundAttribute->GetType(),
+				(int)OptionalAttributes[Idx].Type);
+		}
+	}
+
+	//Verify all attributes are either required or optional, otherwise they are unsupported
+	for (int32 Idx = 0; Idx < AttributesToValidate.Num(); ++Idx)
+	{
+		const FString& Name = AttributesToValidate.GetName(Idx);
+		const FEntry* OptionalAttribute = OptionalAttributes.FindByPredicate([Name](const FEntry& Other) { return Other.Name == Name; });
+		const FEntry* RequiredAttribute = RequiredAttributes.FindByPredicate([Name](const FEntry& Other) { return Other.Name == Name; });
+		
+		if (OptionalAttribute == nullptr && RequiredAttribute == nullptr)
+		{
+			bAreAttributesValid = false;
+			UE_LOG(LogNNX, Warning, TEXT("Found unsupported attribute '%s'."), *Name);
+		}
+	}
+
+	return bAreAttributesValid;
+}
+
+bool FRuntimeRDG::CanCreateModelData(FString FileType, TConstArrayView<uint8> FileData) const
+{
+	return FileType.Compare("onnx", ESearchCase::IgnoreCase) == 0;
+}
+
+bool FRuntimeRDG::CanCreateModel(TConstArrayView<uint8> ModelData) const
+{
+	int32 GuidSize = sizeof(FRuntimeRDG::GUID);
+	int32 VersionSize = sizeof(FRuntimeRDG::Version);
+	if (ModelData.Num() <= GuidSize + VersionSize)
+	{
+		return false;
+	}
+	bool bResult = FGenericPlatformMemory::Memcmp(&(ModelData[0]), &(FRuntimeRDG::GUID), GuidSize) == 0;
+	bResult &= FGenericPlatformMemory::Memcmp(&(ModelData[GuidSize]), &(FRuntimeRDG::Version), VersionSize) == 0;
+	return bResult;
+};
+	
+//
+//
+//
+FInferenceModelRDG::FInferenceModelRDG()
+	: FMLInferenceModel(EMLInferenceModelType::RDG)
+	, bUseManualTransitions(false)
+{
+	
+}
+
+FInferenceModelRDG::~FInferenceModelRDG()
+{
+	
+}
+
+//
+//
+//
+bool FInferenceModelRDG::LoadModel(TConstArrayView<uint8> ModelData, FMLRuntimeFormat& Format)
+{
+	int32 GuidSize = sizeof(FRuntimeRDG::GUID);
+	int32 VersionSize = sizeof(FRuntimeRDG::Version);
+	TConstArrayView<uint8> ModelBuffer = {&(ModelData.GetData()[GuidSize + VersionSize]), ModelData.Num() - GuidSize - VersionSize};
+
+	FMemoryReaderView Reader(ModelBuffer);
+
+	FMLRuntimeFormat::StaticStruct()->SerializeBin(Reader, &Format);
+
+	// Data for base class
+	InputSymbolicTensors.Empty();
+	OutputSymbolicTensors.Empty();
+	
+	// Data for RDG
+	AllSymbolicTensorDescs.Empty();
+	IntermediateTensorIndices.Empty();
+	WeightTensorIndices.Empty();
+	InputTensorIndices.Empty();
+	OutputTensorIndices.Empty();
+	OperatorInputTensorIndices.Empty();
+	OperatorOutputTensorIndices.Empty();
+
+	// Add tensors
+	for (int32 Idx = 0; Idx < Format.Tensors.Num(); ++Idx)
+	{
+		const FMLFormatTensorDesc& FormatTensorDesc = Format.Tensors[Idx];
+
+		const NNECore::FSymbolicTensorShape SymbolicShape = NNECore::FSymbolicTensorShape::Make(FormatTensorDesc.Shape);
+		const NNECore::FTensorDesc SymbolicTensor = NNECore::FTensorDesc::Make(FormatTensorDesc.Name, SymbolicShape, FormatTensorDesc.DataType);
+		
+		AllSymbolicTensorDescs.Emplace(SymbolicTensor);
+		
+		if (FormatTensorDesc.Type == EMLFormatTensorType::Input)
+		{
+			InputTensorIndices.Emplace(Idx);
+			InputSymbolicTensors.Emplace(SymbolicTensor);
+		}
+		else if (FormatTensorDesc.Type == EMLFormatTensorType::Output)
+		{
+			OutputTensorIndices.Emplace(Idx);
+			OutputSymbolicTensors.Emplace(SymbolicTensor);
+		}
+		else if (FormatTensorDesc.Type == EMLFormatTensorType::Intermediate)
+		{
+			IntermediateTensorIndices.Emplace(Idx);
+		}
+		else if (FormatTensorDesc.Type == EMLFormatTensorType::Initializer)
+		{
+			WeightTensorIndices.Emplace(Idx);
+			if (!SymbolicTensor.GetShape().IsConcrete())
+			{
+				UE_LOG(LogNNX, Error, TEXT("Weight tensor %s should have a concrete shape"), *SymbolicTensor.GetName());
+				return false;
+			}
+			
+			const NNECore::FTensorShape TensorShape = NNECore::FTensorShape::MakeFromSymbolic(SymbolicTensor.GetShape());
+			FTensorRDG& WeightRDG = WeightTensorRDGs.Emplace_GetRef(FTensorRDG::Make(SymbolicTensor, TensorShape, nullptr));
+			
+			if (WeightRDG.GetDataSize() != FormatTensorDesc.DataSize)
+			{
+				UE_LOG(LogNNX, Error, TEXT("Weight %s has incorrect size. Expected %d bytes, got %d"), *SymbolicTensor.GetName(), FormatTensorDesc.DataSize, WeightRDG.GetDataSize());
+				return false;
+			}
+
+			const uint8* DataPtr = Format.TensorData.GetData() + FormatTensorDesc.DataOffset;
+			TConstArrayView<uint8> DataView = MakeArrayView(DataPtr, FormatTensorDesc.DataSize);
+
+			WeightRDG.SetPreparedData(DataView);
+		}
+		checkf(FormatTensorDesc.Type != EMLFormatTensorType::None, TEXT("Unsupported tensor type None"));
+	}
+
+	// Loop over all operators in the model and store tensor indices for input/output
+	for (int32 Idx = 0; Idx < Format.Operators.Num(); ++Idx)
+	{
+		OperatorInputTensorIndices.Emplace(Format.Operators[Idx].InTensors);
+		OperatorOutputTensorIndices.Emplace(Format.Operators[Idx].OutTensors);
+	}
+
+	return true;
+}
+
+/**
+ * Run the inference model (synchronous version)
+ */
+int FInferenceModelRDG::RunSync(TConstArrayView<FMLTensorBinding> InInputBindings, TConstArrayView<FMLTensorBinding> InOutputBindings)
+{
+	// Verify the model inputs were prepared
+	if (InputTensorShapes.Num() == 0)
+	{
+		UE_LOG(LogNNX, Error, TEXT("Run(): Input shapes are not set, please call SetInputTensorShapes."));
+		return -1;
+	}
+
+	// TODO ok if pending?
+	Readbacks.Empty();
+	
+	int Res = 0;
+	FEvent* Signal = FGenericPlatformProcess::GetSynchEventFromPool(false);
+
+	ENQUEUE_RENDER_COMMAND(FMLInferenceModel_Run)
+	(
+		[&Signal, &Res, this, InInputBindings, InOutputBindings](FRHICommandListImmediate& RHICmdList)
+		{
+			TOptional<ERHIPipeline>		Pipeline = RHICmdList.GetPipeline();
+
+			if (Pipeline == ERHIPipeline::None)
+			{
+				RHICmdList.SwitchPipeline(ERHIPipeline::Graphics);
+			}
+
+			FRDGBuilder	RDGBuilder(RHICmdList);
+
+			Res = EnqueueRDG(RDGBuilder, InInputBindings, InOutputBindings);
+			if (Res == 0)
+			{
+				RDGBuilder.Execute();
+
+				// FIXME: Using BlockUntilGPUIdle() prevents hang on Linux
+				// FIXME: Adapt to redesigned readback API (UE 5.2)
+				RHICmdList.BlockUntilGPUIdle();
+
+				// Process readback
+				for (const FReadbackEntry& Readback : Readbacks) {
+					const void* BuffData = Readback.RHI->Lock(Readback.Size);
+					check(BuffData);
+					FMemory::Memcpy(Readback.CpuMemory, BuffData, Readback.Size);
+					Readback.RHI->Unlock();
+				}
+			}
+			
+			Signal->Trigger();
+		}
+	);
+
+	// We need to wait for render thread to finish
+	Signal->Wait();
+
+	FGenericPlatformProcess::ReturnSynchEventToPool(Signal);
+
+	return Res;
+}
+
+int FInferenceModelRDG::SetInputTensorShapes(TConstArrayView<FTensorShape> InInputShapes)
+{
+	OutputTensorShapes.Empty();
+	
+	//Verify input shape are valid for the model and set InputTensorShapes
+	if (FMLInferenceModel::SetInputTensorShapes(InInputShapes) != 0)
+	{
+		return -1;
+	}
+
+	//Allocate and prime all AllTensorRDGs with concrete shapes defaulting variables dimension to 1 if needed
+	AllTensorRDGs.Init(nullptr, AllSymbolicTensorDescs.Num());
+
+	InputTensorRDGs.Empty();
+	for (int32 i = 0; i < InputTensorIndices.Num(); ++i)
+	{
+		const int32 Idx = InputTensorIndices[i];
+		const FTensorDesc& TensorDesc = InputSymbolicTensors[i];
+		const FTensorShape& TensorShape = InputTensorShapes[i];
+		
+		InputTensorRDGs.Emplace(FTensorRDG::Make(TensorDesc, TensorShape, nullptr));
+		AllTensorRDGs[Idx] = &InputTensorRDGs[i];
+	}
+
+	for (int32 i = 0; i < WeightTensorIndices.Num(); ++i)
+	{
+		const int32 Idx = WeightTensorIndices[i];
+		
+		AllTensorRDGs[Idx] = &WeightTensorRDGs[i];
+	}
+
+	IntermediateTensorRDGs.Empty();
+	for (int32 i = 0; i < IntermediateTensorIndices.Num(); ++i)
+	{
+		const int32 Idx = IntermediateTensorIndices[i];
+		const FTensorDesc& TensorDesc = AllSymbolicTensorDescs[Idx];
+		const FTensorShape TensorShape = FTensorShape::MakeFromSymbolic(TensorDesc.GetShape());
+		
+		IntermediateTensorRDGs.Emplace(FTensorRDG::Make(TensorDesc, TensorShape, nullptr));
+		AllTensorRDGs[Idx] = &IntermediateTensorRDGs[i];
+	}
+
+	OutputTensorRDGs.Empty();
+	for (int32 i = 0; i < OutputTensorIndices.Num(); ++i)
+	{
+		const int32 Idx = OutputTensorIndices[i];
+		const FTensorDesc& TensorDesc = OutputSymbolicTensors[i];
+		const FTensorShape TensorShape = FTensorShape::MakeFromSymbolic(TensorDesc.GetShape());
+		
+		OutputTensorRDGs.Emplace(FTensorRDG::Make(TensorDesc, TensorShape, nullptr));
+		AllTensorRDGs[Idx] = &OutputTensorRDGs[i];
+	}
+
+	checkCode(
+		for (int i = 0; i < AllTensorRDGs.Num(); ++i)
+		{
+			checkf(AllTensorRDGs[i] != nullptr, TEXT("Tensor at index %d, was not allocated for model preparation."), i);
+		};
+	);
+
+	//Allow the specific engine to run shape inference if supported
+	if (PrepareTensorShapesAndData() != 0)
+	{
+		return -1;
+	}
+
+	checkCode(
+		for (int i = 0; i < AllTensorRDGs.Num(); ++i)
+		{
+			checkf(AllTensorRDGs[i] != nullptr, TEXT("Tensor at index %d, was not allocated after model preparation."), i);
+			checkf(AllTensorRDGs[i]->GetShape().IsCompatibleWith(AllSymbolicTensorDescs[i].GetShape()), TEXT("Tensor at index %d have a shape incompatible with model definition."), i);
+		};
+	);
+
+	//Set OutputTensorShapes for the model from preparation result
+	for (int32 OutputIndices : OutputTensorIndices)
+	{
+		OutputTensorShapes.Emplace(AllTensorRDGs[OutputIndices]->GetShape());
+	}
+
+	check(InputTensorIndices.Num() + OutputTensorIndices.Num() + WeightTensorIndices.Num() + IntermediateTensorIndices.Num() == AllTensorRDGs.Num());
+	check(InputTensorShapes.Num() == InputSymbolicTensors.Num());
+	check(OutputTensorShapes.Num() == OutputSymbolicTensors.Num());
+	check(WeightTensorIndices.Num() == WeightTensorRDGs.Num());
+	check(AllTensorRDGs.Num() == AllSymbolicTensorDescs.Num());
+	
+	return 0;
+}
+
+FRDGBufferDesc CreateRDGBufferDescForTensorRDG(const FTensorRDG& Tensor)
+{
+	// FIXME: CreateStructuredDesc() creates a crash on VulkanRHI
+	//FRDGBufferDesc Desc = FRDGBufferDesc::CreateStructuredDesc(Tensor.GetElemByteSize(), Tensor.GetVolume());
+	FRDGBufferDesc Desc = FRDGBufferDesc::CreateBufferDesc(Tensor.GetElemByteSize(), Tensor.GetVolume());
+
+	return Desc;
+}
+
+/**
+ * Enqueue operators to RDG, the caller will run the GraphBuilder.Execute()
+ */
+int FInferenceModelRDG::EnqueueRDG(FRDGBuilder& RDGBuilder, TConstArrayView<FMLTensorBinding> InInputBindings, TConstArrayView<FMLTensorBinding> InOutputBindings)
+{
+	check(IsInRenderingThread());
+
+	int Res;
+
+	// Verify the model inputs were prepared
+	if (InputTensorShapes.Num() == 0)
+	{
+		UE_LOG(LogNNX, Error, TEXT("EnqueueRDG(): Input shapes are not set, please call SetInputTensorShapes."));
+		return -1;
+	}
+
+	// Process input binding, and if required, allocate RDG buffers
+	FIntArray		InputUploadIndices;
+	Res = SetTensors(RDGBuilder, InputTensorRDGs, InputUploadIndices, InInputBindings);
+	if (Res != 0)
+	{
+		UE_LOG(LogNNX, Warning, TEXT("Invalid input tensor binding type for tensor index:%d"), Res);
+		return -1;
+	}
+
+	// Process output tensors, and if required, allocate RDG buffers
+	FIntArray		OutputReadbackIndices;
+
+	Res = SetTensors(RDGBuilder, OutputTensorRDGs, OutputReadbackIndices, InOutputBindings);
+	if (Res != 0)
+	{
+		UE_LOG(LogNNX, Warning, TEXT("Invalid output tensor binding type for tensor index:%d"), Res);
+		return -1;
+	}
+
+	// If required, upload input tensors to GPU
+	AddTensorUploads_RenderThread(RDGBuilder, InputUploadIndices, InputTensorRDGs, InInputBindings);
+
+	//Create buffer for intermediate tensors
+	for (FTensorRDG& TensorRDG : IntermediateTensorRDGs)
+	{
+		const FRDGBufferDesc BufferDesc = CreateRDGBufferDescForTensorRDG(TensorRDG);
+		const FRDGBufferRef TensorBuffer = RDGBuilder.CreateBuffer(BufferDesc, *TensorRDG.GetName(), ERDGBufferFlags::None);
+		check(TensorRDG.GetBuffer() == nullptr);
+		TensorRDG.SetBuffer(TensorBuffer);
+	}
+
+	// TODO: FIXME: DirectML uses RHI buffers instead of RDG bufers
+	//For now weights tensors are not uploaded to GPU thus GetBuffer will return nullptr for them.
+	//checkCode(for (const FTensorRDG* TensorRDG : AllTensorRDGs) { if (TensorRDG != nullptr) { check(TensorRDG->GetBuffer() != nullptr); } });
+
+	//Insert weights tensors
+	for (int32 i = 0; i < WeightTensorIndices.Num(); ++i)
+	{
+		AllTensorRDGs[WeightTensorIndices[i]] = &WeightTensorRDGs[i];
+	}
+	
+	// We can now dispatch operators
+	AddDispatchOps_RenderThread(RDGBuilder);
+
+	// If required, readback the output tensors to CPU
+	AddTensorReadbacks_RenderThread(RDGBuilder, OutputReadbackIndices, OutputTensorRDGs, InOutputBindings);
+
+	return 0;
+}
+
+/** 
+ * Process binding and check if we need to create RDG Buffer for CPU binding 
+ * Returns 0 on success, or index of a tensor if the tensor type is not supported.
+ */
+int FInferenceModelRDG::SetTensors(FRDGBuilder& GraphBuilder, FTensorRDGArray& InTensorRDGs, FIntArray& OutIndices, 
+	TConstArrayView<FMLTensorBinding> InBindings)
+{
+	check(InBindings.Num() == InTensorRDGs.Num());
+	
+	for (int32 Idx = 0; Idx < InBindings.Num(); ++Idx)
+	{
+		FTensorRDG& TensorRDG = InTensorRDGs[Idx];
+		const FMLTensorBinding& Binding = InBindings[Idx];
+
+		if (Binding.BindingType == EMLTensorBindingDataType::CPUMemory)
+		{
+			FRDGBufferDesc Desc = CreateRDGBufferDescForTensorRDG(TensorRDG);
+			
+			// FIXME: We should use BUF_SourceCopy for only output buffers (GPU readback)
+			Desc.Usage = EBufferUsageFlags(Desc.Usage | BUF_SourceCopy);
+			
+			FRDGBufferRef TensorBuffer = GraphBuilder.CreateBuffer(Desc, *TensorRDG.GetName(), ERDGBufferFlags::None);
+			
+			TensorRDG.SetBuffer(TensorBuffer);
+			OutIndices.Add(Idx);
+		}
+		else if (Binding.BindingType == EMLTensorBindingDataType::RDGBuffer)
+		{
+			TensorRDG.SetBuffer(Binding.Buffer);
+		}
+		else
+		{
+			// Unsupported tensor binding type
+			return Idx;
+		}
+	}
+
+	return 0;
+}
+
+//
+//
+//
+void FInferenceModelRDG::AddTensorUploads_RenderThread(FRDGBuilder& GraphBuilder, TConstArrayView<int32> InUploadIndices, 
+	TConstArrayView<FTensorRDG> InTensorRDGs, TConstArrayView<FMLTensorBinding> InBindings)
+{
+	check(InTensorRDGs.Num() == InBindings.Num());
+	
+	for (int32 Idx = 0; Idx < InUploadIndices.Num(); ++Idx)
+	{
+		const int32			TensorIdx = InUploadIndices[Idx];
+		const FTensorRDG&	TensorRDG = InTensorRDGs[TensorIdx];
+		const FMLTensorBinding& Binding = InBindings[TensorIdx];
+		check(Binding.BindingType == EMLTensorBindingDataType::CPUMemory);
+		GraphBuilder.QueueBufferUpload(TensorRDG.GetBuffer(), Binding.CpuMemory, TensorRDG.GetDataSize(), ERDGInitialDataFlags::NoCopy);
+	}
+}
+
+//
+//
+//
+void FInferenceModelRDG::AddTensorReadbacks_RenderThread(FRDGBuilder& GraphBuilder, TConstArrayView<int32> InReadbackIndices,
+	TConstArrayView<FTensorRDG> InTensorRDGs, TConstArrayView<FMLTensorBinding> InBindings)
+{
+	check(InTensorRDGs.Num() == InBindings.Num());
+
+	check(Readbacks.IsEmpty());
+	
+	for (int32 Idx = 0; Idx < InReadbackIndices.Num(); ++Idx)
+	{
+		const int32			TensorIdx = InReadbackIndices[Idx];
+		const FTensorRDG&	TensorRDG = InTensorRDGs[TensorIdx];
+		const FMLTensorBinding& Binding = InBindings[TensorIdx];
+		check(Binding.BindingType == EMLTensorBindingDataType::CPUMemory);
+		
+		FMLTensorReadbackParameters* TensorReadbackParams = GraphBuilder.AllocParameters<FMLTensorReadbackParameters>();
+
+		TensorReadbackParams->Buffer = TensorRDG.GetBuffer();
+
+		GraphBuilder.AddPass(
+			RDG_EVENT_NAME("FMLInferenceModelAddTensorReadback:%s", *TensorRDG.GetName()),
+			TensorReadbackParams,
+			ERDGPassFlags::Readback | ERDGPassFlags::NeverCull,
+			[this, Binding, TensorRDG, TensorReadbackParams](FRHICommandListImmediate& RHICmdList)
+			{
+				FRHIBuffer* OutputBuffer = TensorReadbackParams->Buffer->GetRHI();
+
+				// TODO: FIXME: We need to transition the resources for DirectML
+				if (bUseManualTransitions)
+				{
+					FRHITransitionInfo Transitions[] =
+					{
+						FRHITransitionInfo(OutputBuffer, ERHIAccess::UAVCompute, ERHIAccess::CopySrc)
+					};
+
+					RHICmdList.Transition(MakeArrayView(Transitions, UE_ARRAY_COUNT(Transitions)));
+					RHICmdList.SubmitCommandsHint();
+				}
+
+				FReadbackEntry& Readback = Readbacks.Add_GetRef({});
+				Readback.RHI = MakeUnique<FRHIGPUBufferReadback>(FName(TEXT("FMLTensorReadback_") + TensorRDG.GetName()));
+				Readback.RHI->EnqueueCopy(RHICmdList, OutputBuffer, TensorRDG.GetDataSize());
+				Readback.CpuMemory = Binding.CpuMemory;
+				Readback.Offset = 0;
+				Readback.Size = TensorRDG.GetDataSize();
+			}
+		);
+	}
+}
+
+TArray<uint8> ConvertToModelData(TArrayView<uint8> ModelBuffer)
+{
+	TArray<uint8> Result;
+
+	FMemoryWriter Writer(Result);
+	Writer << FRuntimeRDG::GUID;
+	Writer << FRuntimeRDG::Version;
+	Writer.Serialize(ModelBuffer.GetData(), ModelBuffer.Num());
+
+	return Result;
+}
+
+} // UE::NNERuntimeRDG::Private
diff --git a/Engine/Plugins/Experimental/NNE/Source/NNERuntimeRDG/Private/NNERuntimeRDG.h b/Engine/Plugins/Experimental/NNE/Source/NNERuntimeRDG/Private/NNERuntimeRDG.h
new file mode 100644
--- /dev/null
+++ b/Engine/Plugins/Experimental/NNE/Source/NNERuntimeRDG/Private/NNERuntimeRDG.h
@@ -0,0 +1,398 @@
+// Copyright Epic Games, Inc. All Rights Reserved.
+
+#pragma once
+
+#include "NNECoreAttributeMap.h"
+#include "NNECoreAttributeValue.h"
+#include "NNECoreTensor.h"
+#include "NNECoreTypes.h"
+#include "NNXCore.h"
+#include "NNXModelOptimizerInterface.h"
+#include "NNXRuntime.h"
+#include "NNXRuntimeFormat.h"
+#include "RenderGraphResources.h"
+#include "Serialization/MemoryReader.h"
+#include "ShaderParameterUtils.h"
+
+#include "Containers/Map.h"
+
+BEGIN_SHADER_PARAMETER_STRUCT(FMLTensorUploadParameters, )
+	RDG_BUFFER_ACCESS(Buffer, ERHIAccess::CopyDest)
+END_SHADER_PARAMETER_STRUCT()
+
+BEGIN_SHADER_PARAMETER_STRUCT(FMLTensorReadbackParameters, )
+	RDG_BUFFER_ACCESS(Buffer, ERHIAccess::CopySrc)
+END_SHADER_PARAMETER_STRUCT()
+
+BEGIN_SHADER_PARAMETER_STRUCT(FMLElementWiseOpParameters, )
+	RDG_BUFFER_ACCESS(InputBuffer, ERHIAccess::UAVCompute)			//!< NOTE: DirectML requires state to be in UAV, even though we're just reading from the InputBuffer
+	RDG_BUFFER_ACCESS(OutputBuffer, ERHIAccess::UAVCompute)
+END_SHADER_PARAMETER_STRUCT()
+
+class FRDGBuilder;
+struct FMLRuntimeFormat;
+struct FNNIModelRaw;
+
+namespace UE::NNECore { class FTensorDesc; }
+
+using namespace NNX;
+
+namespace UE::NNERuntimeRDG::Private
+{
+
+struct FTensorBinding;
+
+/**
+ * Interface for all operators to prepare the model tensors at scheduling time
+ */
+struct IPrepareOperator
+{
+	virtual ~IPrepareOperator() = default;
+	virtual int PrepareOutputs(TConstArrayView<NNECore::Internal::FTensorRef> InputTensors, TArrayView<NNECore::Internal::FTensorRef> OutputTensors) const = 0;
+};
+
+/**
+* Interface for all ML operators running on the RDG 
+*/
+struct IOperatorRDG
+{
+	virtual ~IOperatorRDG() = default;
+};
+
+class FTensorRDG : public NNECore::Internal::FTensor
+{
+	FRDGBufferRef Buffer{};
+	
+public:
+	static FTensorRDG Make(const NNECore::FTensorDesc& TensorDesc, const NNECore::FTensorShape& Shape, FRDGBufferRef Buffer)
+	{
+		check(Shape.IsCompatibleWith(TensorDesc.GetShape()));
+		FTensorRDG TensorRDG;
+		TensorRDG.Buffer = nullptr;
+		TensorRDG.Name = TensorDesc.GetName();
+		TensorRDG.DataType = TensorDesc.GetDataType();
+		TensorRDG.Shape = Shape;
+		TensorRDG.Volume = Shape.Volume();
+		TensorRDG.DataSize = (uint64)NNECore::GetTensorDataTypeSizeInBytes(TensorRDG.DataType) * TensorRDG.Volume;
+		TensorRDG.Buffer = Buffer;
+		check(TensorRDG.Volume <= TNumericLimits<uint32>::Max());
+		return TensorRDG;
+	}
+
+	bool HasBuffer() const { return Buffer != FRDGBufferRef{}; }
+	void SetBuffer(FRDGBufferRef Inbuffer){ Buffer = Inbuffer; }
+	FRDGBufferRef GetBuffer() const { return Buffer; }
+};
+
+using FTensorRDGRef = FTensorRDG*;
+using FTensorRDGArray = TArray<FTensorRDG, TInlineAllocator<16>>;
+using FTensorRDGRefArray = TArray<FTensorRDGRef, TInlineAllocator<64>>;
+using FIntArray = TArray<int32, TInlineAllocator<16>>;
+
+class FRuntimeRDG : public IRuntime
+{
+public:
+	static FGuid GUID;
+	static int32 Version;
+
+	virtual bool CanCreateModelData(FString FileType, TConstArrayView<uint8> FileData) const;
+	virtual TArray<uint8> CreateModelData(FString FileType, TConstArrayView<uint8> FileData) = 0;
+	virtual bool CanCreateModel(TConstArrayView<uint8> ModelData) const;
+};
+
+/** 
+ * RDG inference model base class
+ */
+class FInferenceModelRDG : public FMLInferenceModel
+{
+
+	struct FReadbackEntry;
+
+public:
+
+	~FInferenceModelRDG();
+
+	virtual int SetInputTensorShapes(TConstArrayView<FTensorShape> InputShapes) override;
+	virtual int RunSync(TConstArrayView<FMLTensorBinding> InputBindings, TConstArrayView<FMLTensorBinding> OutputBindings) override;
+	virtual int EnqueueRDG(FRDGBuilder& RDGBuilder, TConstArrayView<FMLTensorBinding> InputBindings, TConstArrayView<FMLTensorBinding> OutputBindings) override;
+
+protected:
+
+	FInferenceModelRDG();
+
+	bool LoadModel(TConstArrayView<uint8> ModelData, FMLRuntimeFormat& Format);
+
+	int SetTensors(FRDGBuilder& GraphBuilder, FTensorRDGArray& InTensorRDGs, FIntArray& OutIndices, TConstArrayView<FMLTensorBinding> InBindings);
+	
+	virtual int PrepareTensorShapesAndData() = 0;
+	virtual void AddDispatchOps_RenderThread(FRDGBuilder& GraphBuilder) = 0;
+
+	virtual void AddTensorUploads_RenderThread(FRDGBuilder& GraphBuilder, TConstArrayView<int32> InUploadIndices, TConstArrayView<FTensorRDG> InRDGBindings, TConstArrayView<FMLTensorBinding> InBindings);
+	virtual void AddTensorReadbacks_RenderThread(FRDGBuilder& GraphBuilder, TConstArrayView<int32> InReadbackIndices, TConstArrayView<FTensorRDG> InRDGBindings, TConstArrayView<FMLTensorBinding> InBindings);
+
+	//Tensor descriptor
+	TArray<FTensorDesc>			AllSymbolicTensorDescs;
+
+	//Tensor indices for models
+	TArray<int32>				IntermediateTensorIndices;
+	TArray<int32>				WeightTensorIndices;
+	TArray<int32>				InputTensorIndices;
+	TArray<int32>				OutputTensorIndices;
+	
+	//Tensor indices by operator
+	TArray<TArray<uint32>>		OperatorInputTensorIndices;
+	TArray<TArray<uint32>>		OperatorOutputTensorIndices;
+
+	//RDG Tensors
+	FTensorRDGRefArray			AllTensorRDGs;
+	FTensorRDGArray				InputTensorRDGs;
+	FTensorRDGArray				OutputTensorRDGs;
+	FTensorRDGArray				IntermediateTensorRDGs;
+	FTensorRDGArray				WeightTensorRDGs;
+	
+	TArray<FReadbackEntry>		Readbacks;
+	bool						bUseManualTransitions;
+};
+
+//TODO jira 167584 remove default validation and declare contract in all DML operator (see HLSL Gemm for current example)
+bool AlwaysValidValidationFunction(const NNECore::FAttributeMap& AttributeMap, TConstArrayView<ENNETensorDataType> InputTypes, TConstArrayView<NNECore::FSymbolicTensorShape> InputShapes);
+
+class FInputValidator
+{
+public:
+	FInputValidator();
+	void AddRequired(int32 TemplateIdx = 0);
+	void AddOptional(int32 TemplateIdx = 0);
+	void SetTemplateCount(int32 TemplateCount);
+	void AddSupportedType(ENNETensorDataType Type, int32 TemplateIdx = 0);
+	bool Validate(TConstArrayView<ENNETensorDataType> InputTypes);
+
+private:
+	TArray<TArray<ENNETensorDataType>> TemplateTypes;
+	TArray<int32> InputTemplateIndices;
+	int32 NumRequiredInput;
+	int32 NumOptionalInput;
+};
+
+class FAttributeValidator
+{
+public:
+	void AddOptional(const FString& Name, ENNEAttributeDataType Type);
+	void AddRequired(const FString& Name, ENNEAttributeDataType Type);
+	bool Validate(const NNECore::FAttributeMap& AttributesToValidate);
+
+private:
+	struct FEntry
+	{
+		FEntry(const FString& InName, ENNEAttributeDataType InType)
+			: Name(InName), Type(InType)
+		{
+		}
+
+		//Idea: we could extended as needed by operator to support more validation especially around the range of the values.
+		//An example is ConvTranspose `auto_pad` enum style string that can only take a few values
+		//In the same direction we might only support a range of value for a float (for example
+		//we only support integer but the type is float, or only positive values for an int32)
+		FString Name;
+		ENNEAttributeDataType Type;
+	};
+
+	TArray<FEntry> RequiredAttributes;
+	TArray<FEntry> OptionalAttributes;
+};
+
+/**
+ * Registry for RDG ML operators
+ */
+template<class TOperatorType>
+class TOperatorRegistryRDG
+{
+public:
+
+	typedef TOperatorType* (*OperatorCreateFunc)();
+	typedef bool (*OperatorValidateFunc)(const NNECore::FAttributeMap& AttributeMap, TConstArrayView<ENNETensorDataType> InputTypes, TConstArrayView<NNECore::FSymbolicTensorShape> InputShapes);
+
+	static TOperatorRegistryRDG* Get()
+	{
+		static TOperatorRegistryRDG Inst;
+
+		return &Inst;
+	}
+
+	OperatorValidateFunc OpFindValidation(const FString& Name)
+	{
+		OperatorValidateFunc* Fn = OperatorValidations.Find(Name);
+
+		if (!Fn)
+		{
+			UE_LOG(LogNNX, Warning, TEXT("RDG MLOperator:%s is not registered"), *Name);
+			return nullptr;
+		}
+
+		return *Fn;
+	}
+
+	OperatorCreateFunc OpFind(const FString& Name)
+	{
+		OperatorCreateFunc* Fn = Operators.Find(Name);
+
+		if (!Fn)
+		{
+			UE_LOG(LogNNX, Warning, TEXT("RDG MLOperator:%s is not registered"), *Name);
+			return nullptr;
+		}
+
+		return *Fn;
+	}
+
+	bool OpAdd(const FString& Name, OperatorCreateFunc Func, OperatorValidateFunc ValidateFunc = AlwaysValidValidationFunction)
+	{
+		if (Operators.Find(Name) != nullptr)
+		{
+			UE_LOG(LogNNX, Warning, TEXT("RDG MLOperator is already registered:%s"), *Name);
+			return false;
+		}
+
+		Operators.Add(Name, Func);
+		OperatorValidations.Add(Name, ValidateFunc);
+		return true;
+	}
+
+private:
+
+	TMap<FString, OperatorCreateFunc> Operators;
+	TMap<FString, OperatorValidateFunc> OperatorValidations;
+};
+
+/**
+ * Validator for RDG ML operators
+ */
+template<class TOperatorType>
+class TModelValidatorRDG : public IModelValidator
+{
+public:
+	virtual FString GetName() const 
+	{
+		return TEXT("RDG Model validator");
+	}
+	
+	virtual bool ValidateModel(const FNNIModelRaw& InputModel, const FOptimizerOptionsMap& Options) const override
+	{
+		FMLRuntimeFormat	Format;
+
+		ENNXInferenceFormat FormatType = InputModel.Format;
+		if (FormatType != ENNXInferenceFormat::NNXRT)
+		{
+			UE_LOG(LogNNX, Warning, TEXT("Unsupported format type for validator %s"), *GetName());
+			return false;
+		}
+
+		FMemoryReader Reader(InputModel.Data);
+		FMLRuntimeFormat::StaticStruct()->SerializeBin(Reader, &Format);
+
+		TOperatorRegistryRDG<TOperatorType>* Registry = TOperatorRegistryRDG<TOperatorType>::Get();
+		check(Registry != nullptr);
+
+		for (int32 Idx = 0; Idx < Format.Operators.Num(); ++Idx)
+		{
+			TArray<ENNETensorDataType> InputTensorTypes;
+			TArray<NNECore::FSymbolicTensorShape> InputTensorShapes;
+			NNECore::FAttributeMap AttributeMap;
+			
+			for (int32 InputTensorIndex: Format.Operators[Idx].InTensors)
+			{
+				InputTensorTypes.Add(Format.Tensors[InputTensorIndex].DataType);
+				InputTensorShapes.Add(NNECore::FSymbolicTensorShape::Make(Format.Tensors[InputTensorIndex].Shape));
+			}
+			for (const FMLFormatAttributeDesc& Desc : Format.Operators[Idx].Attributes)
+			{
+				AttributeMap.SetAttribute(Desc.Name, Desc.Value);
+			}
+
+			const FString& OpType = Format.Operators[Idx].TypeName;
+			
+			typename TOperatorRegistryRDG<TOperatorType>::OperatorValidateFunc ValidationFn = Registry->OpFindValidation(OpType);
+
+			if (!ValidationFn)
+			{
+				UE_LOG(LogNNX, Warning, TEXT("Hlsl MLOperatorRegistry failed to find validation for operator:%s"), *OpType);
+				return false;
+			}
+			
+			if (!ValidationFn(AttributeMap, InputTensorTypes, InputTensorShapes))
+			{
+				UE_LOG(LogNNX, Warning, TEXT("Hlsl MLOperatorRegistry failed to validate operator:%s"), *OpType);
+				return false;
+			}
+		}
+
+		return true;
+	}
+};
+
+TArray<uint8> ConvertToModelData(TArrayView<uint8> ModelBuffer);
+
+} // UE::NNERuntimeRDG::Private
+
+namespace UE::NNERuntimeRDG::Private::Dml
+{
+
+// NOTE: For now we only have DML on Windows, we should add support for XSX
+#ifdef NNE_USE_DIRECTML
+
+	extern IRuntime* FRuntimeDmlStartup();
+	extern void FRuntimeDmlShutdown();
+
+#endif
+	
+}
+
+namespace UE::NNERuntimeRDG::Private::Hlsl
+{
+
+	extern IRuntime* FRuntimeHlslStartup();
+	extern void FRuntimeHlslShutdown();
+}
+
+///// OperatorId
+//class FOperatorId
+//{
+//public:
+//
+//	FOperatorId(uint32 InType, uint32 InIndex)
+//		: Type(InType)
+//		, Index(InIndex)
+//	{
+//	}
+//
+//	uint32 GetIndex() const
+//	{
+//		return Index;
+//	}
+//
+//	uint32 GetType() const
+//	{
+//		return Type;
+//	}
+//
+//private:
+//
+//	uint32	Type : 8;
+//	uint32	Index : 16;
+//};
+
+///// Base class
+//class FRDGOperator
+//{};
+
+///// Base class for RDG runtime
+//class FRDGRuntime : public IRuntime
+//{
+//public:
+//
+//	virtual FRDGOperator OpCreate(const FString& Name, const TArrayView<uint32>& TensorSizes) = 0;
+//
+//};
+
+
+ 
\ No newline at end of file
diff --git a/Engine/Plugins/Experimental/NNE/Source/NNERuntimeRDG/Private/NNERuntimeRDGDml.cpp b/Engine/Plugins/Experimental/NNE/Source/NNERuntimeRDG/Private/NNERuntimeRDGDml.cpp
new file mode 100644
--- /dev/null
+++ b/Engine/Plugins/Experimental/NNE/Source/NNERuntimeRDG/Private/NNERuntimeRDGDml.cpp
@@ -0,0 +1,2369 @@
+// Copyright Epic Games, Inc. All Rights Reserved.
+
+#include "NNERuntimeRDG.h"
+#include "NNXInferenceModel.h"
+#include "NNXRuntimeFormat.h"
+#include "NNEUtilsModelOptimizer.h"
+#include "NNECoreOperator.h"
+
+#include "NNECoreAttributeMap.h"
+
+#include "HAL/FileManager.h"
+
+#include "RenderGraphBuilder.h"
+#include "RenderGraphUtils.h"
+
+#include "Algo/Find.h"
+
+#ifdef NNE_USE_DIRECTML
+
+#include "Windows/AllowWindowsPlatformTypes.h"
+THIRD_PARTY_INCLUDES_START
+#include <unknwn.h>
+#include "Microsoft/COMPointer.h"
+#include "DirectML.h"
+THIRD_PARTY_INCLUDES_END
+#include "Windows/HideWindowsPlatformTypes.h"
+
+// DirectML is implemented using COM on all platforms
+#ifdef IID_GRAPHICS_PPV_ARGS
+#define DML_PPV_ARGS(x) __uuidof(*x), IID_PPV_ARGS_Helper(x)
+#else
+#define DML_PPV_ARGS(x) IID_PPV_ARGS(x)
+#endif
+
+#include "ID3D12DynamicRHI.h"
+
+#define NNE_USE_D3D12_RESOURCES
+#define NNX_RUNTIME_DML_NAME TEXT("NNXRuntimeDml")
+
+
+namespace UE::NNERuntimeRDG::Private::Dml
+{
+/**
+ * Utility function to get operator type
+ */
+template<typename TElementWiseOpDesc>
+DML_OPERATOR_TYPE GetDmlElementWiseUnaryOpType();
+
+#define OP_EW(OpName) \
+template<> \
+DML_OPERATOR_TYPE GetDmlElementWiseUnaryOpType<DML_ELEMENT_WISE_##OpName##_OPERATOR_DESC>() { return DML_OPERATOR_ELEMENT_WISE_##OpName; }
+
+#define OP_AN(OpName) \
+template<> \
+DML_OPERATOR_TYPE GetDmlElementWiseUnaryOpType<DML_ACTIVATION_##OpName##_OPERATOR_DESC>() { return DML_OPERATOR_ACTIVATION_##OpName; }
+
+
+OP_EW(IDENTITY)
+OP_EW(ABS)
+OP_EW(ACOS)
+OP_EW(ACOSH)
+OP_EW(ASIN)
+OP_EW(ASINH)
+OP_EW(ATAN)
+OP_EW(ATANH)
+// BitShift
+// Cast
+OP_EW(CEIL)
+OP_EW(CLIP)
+OP_EW(COS)
+OP_EW(COSH)
+OP_AN(ELU)
+OP_EW(ERF)
+OP_EW(EXP)
+OP_EW(FLOOR)
+OP_EW(IS_INFINITY)
+OP_EW(IS_NAN)
+OP_AN(HARDMAX)
+OP_AN(HARD_SIGMOID)
+OP_AN(LEAKY_RELU)
+OP_EW(LOG)
+OP_EW(NEGATE)
+// Not
+OP_EW(RECIP)
+OP_AN(RELU)
+OP_EW(ROUND)
+OP_AN(SCALED_ELU)
+OP_AN(SIGMOID)
+OP_EW(SIGN)
+OP_EW(SIN)
+OP_EW(SINH)
+OP_AN(SOFTPLUS)
+OP_AN(SOFTSIGN)
+OP_EW(SQRT)
+OP_EW(TAN)
+OP_EW(TANH)
+
+#undef OP_EW
+#undef OP_AN
+
+
+/**
+ * Utility function to get operator type
+ */
+//template<typename TActivationOpDesc>
+//DML_OPERATOR_TYPE GetDmlActivationOpType();
+//
+//#define OP(OpName) \
+//template<> \
+//DML_OPERATOR_TYPE GetDmlActivationOpType<DML_ACTIVATION_##OpName##_OPERATOR_DESC>() { return DML_OPERATOR_ACTIVATION_##OpName; }
+//
+//OP(ELU)
+//OP(ERF)
+//OP(HARDMAX)
+//OP(HARD_SIGMOID)
+//OP(LEAKY_RELU)
+//OP(LINEAR)
+//OP(LOG_SOFTMAX)
+//OP(PARAMETERIZED_RELU)
+//OP(PARAMETRIC_SOFTPLUS)
+//OP(RELU)
+//OP(SCALED_ELU)
+//OP(SCALED_TANH)
+//OP(SOFTMAX)
+//OP(SOFTPLUS)
+//OP(SOFTSIGN)
+//OP(TANH)
+//OP(THRESHOLDED_RELU)
+//
+
+/**
+ * Utility function to get operator type
+ */
+template<typename TElementWiseOpDesc>
+DML_OPERATOR_TYPE GetDmlElementWiseBinaryOpType();
+
+#define OP_EW(OpName) \
+template<> \
+DML_OPERATOR_TYPE GetDmlElementWiseBinaryOpType<DML_ELEMENT_WISE_##OpName##_OPERATOR_DESC>() { return DML_OPERATOR_ELEMENT_WISE_##OpName; }
+
+#define OP_AN(OpName) \
+template<> \
+DML_OPERATOR_TYPE GetDmlElementWiseBinaryOpType<DML_ACTIVATION_##OpName##_OPERATOR_DESC>() { return DML_OPERATOR_ACTIVATION_##OpName; }
+
+OP_EW(ADD)
+//OP_EW(LOGICAL_AND)
+OP_EW(DIVIDE)
+//OP_EW(LOGICAL_EQUALS)
+//OP_EW(LOGICAL_GREATER_THAN)
+//OP_EW(LOGICAL_LESS_THAN)
+//OP_EW(MOD)
+OP_EW(MULTIPLY)
+//OP_EW(LOGICAL_OR)
+OP_AN(PARAMETERIZED_RELU)
+OP_EW(POW)
+OP_EW(SUBTRACT)
+//OP_EW(LOGICAL_XOR)
+
+#undef OP_EW
+#undef OP_AN
+
+namespace DmlUtil
+{
+	struct FTensorDesc
+	{
+		DML_BUFFER_TENSOR_DESC										BuffDesc;
+		DML_TENSOR_DESC												Desc;
+		TArray<uint32, TInlineAllocator<FTensorShape::MaxRank>>		Sizes;
+		TArray<uint32, TInlineAllocator<FTensorShape::MaxRank>>		Strides;
+	};
+
+	void SetTensorStrides(FTensorDesc& TensorDesc, const NNECore::Internal::FTensor& InputDesc)
+	{
+		uint32 CurrStride = 1;
+
+		TensorDesc.Strides.SetNum(InputDesc.GetShape().Rank());
+		
+		for (int32 i = InputDesc.GetShape().Rank() - 1; i >= 0; --i)
+		{
+			TensorDesc.Strides[i] = CurrStride;
+			CurrStride *= InputDesc.GetShape().GetData()[i];
+		}
+	}
+
+	void SetTensorSizesAndStridesForBroadcast(FTensorDesc& TensorDesc, const NNECore::Internal::FTensor& InputDesc, const NNECore::Internal::FTensor& TargetDesc)
+	{
+		static_assert(FTensorShape::MaxRank <= 8);
+		
+		const uint32 TargetDimension = TargetDesc.GetShape().Rank() != -1 ? TargetDesc.GetShape().Rank() : InputDesc.GetShape().Rank();
+		checkf(TargetDesc.GetShape().Rank() >= InputDesc.GetShape().Rank(), TEXT("Can't broadcast tensor from rank %d to rank %d, should be inferior or equal."), InputDesc.GetShape().Rank(), TargetDimension);
+		
+		TensorDesc.Sizes.SetNum(TargetDimension);
+		TensorDesc.Strides.SetNum(TargetDimension);
+
+		const int32 DimensionOffset = int32(TargetDimension - InputDesc.GetShape().Rank());
+		
+		for (int32 i = 0; i < (int32) TargetDimension; ++i)
+		{
+			TensorDesc.Sizes[i] = i < DimensionOffset ? 1 : InputDesc.GetShape().GetData()[i - DimensionOffset];
+		}
+
+		uint32 CurrStride = 1;
+
+		for (int32 i = TargetDimension - 1; i >= 0; --i)
+		{
+			const bool bBroadcast = TensorDesc.Sizes[i] < TargetDesc.GetShape().GetData()[i];
+
+			TensorDesc.Strides[i] = bBroadcast ? 0 : CurrStride;
+			CurrStride *= TensorDesc.Sizes[i];
+
+			TensorDesc.Sizes[i] = TargetDesc.GetShape().GetData()[i];
+		}
+	}
+
+	inline bool IsSameShape(const NNECore::Internal::FTensor& Left, const NNECore::Internal::FTensor& Right)
+	{
+		if (Left.GetShape().Rank() != Right.GetShape().Rank())
+		{
+			return false;
+		}
+		
+		for (int32 Idx = 0; Idx < Left.GetShape().Rank(); ++Idx)
+		{
+			if (Left.GetShape().GetData()[Idx] != Right.GetShape().GetData()[Idx])
+			{
+				return false;
+			}
+		}
+
+		return true;
+	}
+
+	DML_TENSOR_DATA_TYPE GetTensorDataType(ENNETensorDataType DataType)
+	{
+		switch (DataType)
+		{
+		case ENNETensorDataType::Double:
+			return DML_TENSOR_DATA_TYPE_FLOAT64;
+
+		case ENNETensorDataType::Float:
+			return DML_TENSOR_DATA_TYPE_FLOAT32;
+
+		case ENNETensorDataType::Half:
+			return DML_TENSOR_DATA_TYPE_FLOAT16;
+
+		case ENNETensorDataType::UInt64:
+			return DML_TENSOR_DATA_TYPE_UINT64;
+
+		case ENNETensorDataType::UInt32:
+			return DML_TENSOR_DATA_TYPE_UINT32;
+
+		case ENNETensorDataType::UInt16:
+			return DML_TENSOR_DATA_TYPE_UINT16;
+
+		case ENNETensorDataType::UInt8:
+			return DML_TENSOR_DATA_TYPE_UINT8;
+
+		case ENNETensorDataType::Int64:
+			return DML_TENSOR_DATA_TYPE_INT64;
+
+		case ENNETensorDataType::Int32:
+			return DML_TENSOR_DATA_TYPE_INT32;
+
+		case ENNETensorDataType::Int16:
+			return DML_TENSOR_DATA_TYPE_INT16;
+
+		case ENNETensorDataType::Int8:
+			return DML_TENSOR_DATA_TYPE_INT8;
+
+		default:
+			return DML_TENSOR_DATA_TYPE_UNKNOWN;
+		}
+	}
+
+	inline uint64 CalculateBufferSize(const FTensorDesc& DmlTensor, const NNECore::Internal::FTensor& Desc)
+	{
+		uint64 ElemSizeInBytes = Desc.GetElemByteSize();
+		
+		if (ElemSizeInBytes == 0)
+		{
+			return 0;
+		}
+
+		uint64 MinSizeInBytes = 0;		
+		uint32 IndexOfLastElement = 0;
+
+		for (int32 i = 0; i < DmlTensor.Sizes.Num(); ++i)
+		{
+			IndexOfLastElement += (DmlTensor.Sizes[i] - 1) * DmlTensor.Strides[i];
+		}
+
+		MinSizeInBytes = (static_cast<uint64>(IndexOfLastElement) + 1) * ElemSizeInBytes;
+		
+		// Round up to the nearest 4 bytes.
+		MinSizeInBytes = (MinSizeInBytes + 3) & ~3ull;
+
+		return MinSizeInBytes;
+	}
+}
+
+//
+//
+//
+class FDeviceContextDml
+{
+public:
+
+	uint32							DeviceIndex;
+	ID3D12Device*					D3D12Device{ nullptr }; // Borrowed reference from RHI
+	TComPtr<IDMLDevice>				Device{ nullptr };
+	TComPtr<IDMLCommandRecorder>	CmdRec{ nullptr };
+};
+
+//
+// DirectML operator
+//
+class FOperatorDml
+{
+public:
+
+	virtual ~FOperatorDml() = default;
+
+	virtual bool Initialize(FDeviceContextDml* DevCtx, TArrayView<const NNECore::Internal::FTensor> InputTensors, TArrayView<const NNECore::Internal::FTensor> OutputTensors, const NNECore::FAttributeMap& Attributes) = 0;
+
+	IDMLOperator* GetOperator()
+	{
+		return DmlOp;
+	}
+
+protected:
+
+	bool InitDmlTensorDesc(DmlUtil::FTensorDesc& DmlTensorDesc, const NNECore::Internal::FTensor& TensorDesc)
+	{
+		DML_TENSOR_DATA_TYPE DmlDataType = DmlUtil::GetTensorDataType(TensorDesc.GetDataType());
+
+		if (DmlDataType == DML_TENSOR_DATA_TYPE_UNKNOWN)
+		{
+			DmlTensorDesc.BuffDesc = DML_BUFFER_TENSOR_DESC{};
+			DmlTensorDesc.Desc = DML_TENSOR_DESC{};
+
+			return false;
+		}
+
+		DmlTensorDesc.Sizes = TensorDesc.GetShape().GetData();
+		// TODO: Support tensor padding using strides defined in FMLTensorDesc
+		//DmlUtil::SetTensorStrides(DmlTensorDesc, TensorDesc.Strides);
+		
+		DML_BUFFER_TENSOR_DESC& BuffDesc = DmlTensorDesc.BuffDesc;
+
+		BuffDesc = DML_BUFFER_TENSOR_DESC{};
+		BuffDesc.DataType = DmlDataType;
+		BuffDesc.Flags = TensorDesc.HasPreparedData() ? DML_TENSOR_FLAG_OWNED_BY_DML : DML_TENSOR_FLAG_NONE;
+		BuffDesc.DimensionCount = TensorDesc.GetShape().Rank();
+		BuffDesc.Sizes = DmlTensorDesc.Sizes.GetData();
+		BuffDesc.Strides = nullptr;
+		BuffDesc.TotalTensorSizeInBytes = TensorDesc.GetDataSize();
+
+		DmlTensorDesc.Desc = DML_TENSOR_DESC{ DML_TENSOR_TYPE_BUFFER, &DmlTensorDesc.BuffDesc };
+
+		return true;
+	}
+
+	bool InitDmlTensorDesc(DmlUtil::FTensorDesc& DmlTensorDesc, const NNECore::Internal::FTensor& TensorDesc, const NNECore::Internal::FTensor& BroadcastDesc)
+	{
+		DML_TENSOR_DATA_TYPE DmlDataType = DmlUtil::GetTensorDataType(TensorDesc.GetDataType());
+
+		if (DmlDataType == DML_TENSOR_DATA_TYPE_UNKNOWN)
+		{
+			DmlTensorDesc.BuffDesc = DML_BUFFER_TENSOR_DESC{};
+			DmlTensorDesc.Desc = DML_TENSOR_DESC{};
+
+			return false;
+		}
+
+		if (DmlUtil::IsSameShape(TensorDesc, BroadcastDesc))
+		{
+			DmlTensorDesc.Sizes = TensorDesc.GetShape().GetData();
+			DmlUtil::SetTensorStrides(DmlTensorDesc, TensorDesc);
+		}
+		else if (TensorDesc.GetShape().Rank() > BroadcastDesc.GetShape().Rank())
+		{
+			return false;
+		}
+		else
+		{
+			DmlUtil::SetTensorSizesAndStridesForBroadcast(DmlTensorDesc, TensorDesc, BroadcastDesc);
+		}
+
+		//UE_LOG(LogNNX, Warning, TEXT("DmlTensorDesc:%d,%d,%d -> %d,%d,%d"),
+		//	TensorDesc.Sizes[0],
+		//	TensorDesc.Dimension > 1 ? TensorDesc.Sizes[1] : 0,
+		//	TensorDesc.Dimension > 2 ? TensorDesc.Sizes[2] : 0,
+		//	DmlTensorDesc.Sizes[0],
+		//	DmlTensorDesc.Dimension > 1 ? DmlTensorDesc.Sizes[1] : 0,
+		//	DmlTensorDesc.Dimension > 2 ? DmlTensorDesc.Sizes[2] : 0
+		//);
+
+		//UE_LOG(LogNNX, Warning, TEXT("DmlTensorStrides:%d,%d,%d"),
+		//	DmlTensorDesc.Strides[0], DmlTensorDesc.Strides[1], DmlTensorDesc.Strides[2]
+		//);
+
+		check(DmlTensorDesc.Strides.Num() == DmlTensorDesc.Sizes.Num());
+		
+		DML_BUFFER_TENSOR_DESC& BuffDesc = DmlTensorDesc.BuffDesc;
+		
+		BuffDesc = DML_BUFFER_TENSOR_DESC{};
+
+		BuffDesc.DataType = DmlDataType;
+		BuffDesc.Flags = TensorDesc.HasPreparedData() ? DML_TENSOR_FLAG_OWNED_BY_DML : DML_TENSOR_FLAG_NONE;
+		BuffDesc.DimensionCount = DmlTensorDesc.Sizes.Num();
+		BuffDesc.Sizes = DmlTensorDesc.Sizes.GetData();
+		BuffDesc.Strides = DmlTensorDesc.Strides.GetData();
+		BuffDesc.TotalTensorSizeInBytes = DmlUtil::CalculateBufferSize(DmlTensorDesc, TensorDesc);
+		
+		DmlTensorDesc.Desc = DML_TENSOR_DESC{ DML_TENSOR_TYPE_BUFFER, &DmlTensorDesc.BuffDesc };
+
+		return true;
+	}
+	
+	bool CreateOperator(const DML_OPERATOR_DESC& DmlOpDesc)
+	{
+		IDMLDevice* Device = DevCtx->Device;
+
+		// Create operator
+		IDMLOperator* Op = nullptr;
+
+		HRESULT Res;
+
+		Res = Device->CreateOperator(&DmlOpDesc, DML_PPV_ARGS(&Op));
+		if (!Op)
+		{
+			UE_LOG(LogNNX, Warning, TEXT("Error:Failed to create DML operator, hres:%d"), Res);
+			return false;
+		}
+
+		DmlOp = Op;
+
+		return DmlOp.IsValid();
+	}
+
+	FDeviceContextDml*			DevCtx;
+	TComPtr<IDMLOperator>		DmlOp;
+};
+
+/**
+ * DirectML ML operator registry
+ */
+using FOperatorRegistryDml = TOperatorRegistryRDG<FOperatorDml>;
+
+/**
+ * Element-wise unary ML operator implementation
+ */
+template
+<
+	typename DmlElementWiseOpDescType, 
+	NNECore::Internal::EElementWiseUnaryOperatorType OpType
+>
+class FOperatorDmlElementWiseUnary : public FOperatorDml
+{
+public:
+
+	static FOperatorDml* Create()
+	{
+		return new FOperatorDmlElementWiseUnary();
+	}
+
+	virtual ~FOperatorDmlElementWiseUnary() = default;
+
+private:
+
+	FOperatorDmlElementWiseUnary() : Alpha(0.0f), Beta(0.0f), Gamma(0.0f), Num(1) {}
+	float Alpha;
+	float Beta;
+	float Gamma;
+	uint32 Num;
+
+public:
+
+	//
+	//
+	//
+	virtual bool Initialize(FDeviceContextDml* InDevCtx, TArrayView<const NNECore::Internal::FTensor> InputTensors, TArrayView<const NNECore::Internal::FTensor> OutputTensors, const NNECore::FAttributeMap& Attributes) override
+	{
+		Num = InputTensors[0].GetVolume();
+
+		DevCtx = InDevCtx;
+
+		const NNECore::Internal::FTensor& InputTensorDesc = InputTensors[0];
+		const NNECore::Internal::FTensor& OutputTensorDesc = OutputTensors[0];
+
+		Alpha = Attributes.GetValueOrDefault(TEXT("alpha"), Alpha);
+		Beta = Attributes.GetValueOrDefault(TEXT("beta"), Beta);
+		Gamma = Attributes.GetValueOrDefault(TEXT("gamma"), Gamma);
+
+		// Initialize tensor descriptor (it's same for both input and output)
+		DmlUtil::FTensorDesc	DmlTensorDesc{};
+
+		if (!InitDmlTensorDesc(DmlTensorDesc, InputTensorDesc))
+		{
+			UE_LOG(LogNNX, Warning, TEXT("Failed to initialize tensor(s) for DML inference"));
+			return false;
+		}
+
+		DmlElementWiseOpDescType	DmlElemWiseOpDesc{};
+
+		InitDmlOpDesc(DmlElemWiseOpDesc, DmlTensorDesc);
+
+		DML_OPERATOR_DESC DmlOpDesc{};
+
+		DmlOpDesc.Type = GetDmlElementWiseUnaryOpType<DmlElementWiseOpDescType>();
+		DmlOpDesc.Desc = &DmlElemWiseOpDesc;
+
+		return CreateOperator(DmlOpDesc);
+	}
+
+private:
+
+	template<typename OpDesc>
+	void InitDmlOpDesc(OpDesc& Desc, DmlUtil::FTensorDesc& TensorDesc)
+	{
+		Desc.InputTensor = &TensorDesc.Desc;
+		Desc.OutputTensor = &TensorDesc.Desc;
+	}
+
+	void InitDmlOpDesc(DML_ACTIVATION_SOFTPLUS_OPERATOR_DESC& Desc, DmlUtil::FTensorDesc& TensorDesc)
+	{
+		Desc.InputTensor = &TensorDesc.Desc;
+		Desc.OutputTensor = &TensorDesc.Desc;
+		Desc.Steepness = 1.0f;
+	}
+
+	void InitDmlOpDesc(DML_ACTIVATION_SCALED_ELU_OPERATOR_DESC& Desc, DmlUtil::FTensorDesc& TensorDesc)
+	{
+		Desc.InputTensor = &TensorDesc.Desc;
+		Desc.OutputTensor = &TensorDesc.Desc;
+		Desc.Alpha = Alpha;
+		Desc.Gamma = Gamma;
+	}
+
+	void InitDmlOpDesc(DML_ACTIVATION_ELU_OPERATOR_DESC& Desc, DmlUtil::FTensorDesc& TensorDesc)
+	{
+		Desc.InputTensor = &TensorDesc.Desc;
+		Desc.OutputTensor = &TensorDesc.Desc;
+		Desc.Alpha = Alpha;
+	}
+
+	void InitDmlOpDesc(DML_ACTIVATION_HARD_SIGMOID_OPERATOR_DESC& Desc, DmlUtil::FTensorDesc& TensorDesc)
+	{
+		Desc.InputTensor = &TensorDesc.Desc;
+		Desc.OutputTensor = &TensorDesc.Desc;
+		Desc.Alpha = Alpha;
+		Desc.Beta = Beta;
+	}
+
+	void InitDmlOpDesc(DML_ACTIVATION_LEAKY_RELU_OPERATOR_DESC& Desc, DmlUtil::FTensorDesc& TensorDesc)
+	{
+		Desc.InputTensor = &TensorDesc.Desc;
+		Desc.OutputTensor = &TensorDesc.Desc;
+		Desc.Alpha = Alpha;
+	}
+};
+
+template<> FOperatorDmlElementWiseUnary<DML_ACTIVATION_SCALED_ELU_OPERATOR_DESC, NNECore::Internal::EElementWiseUnaryOperatorType::Selu>::FOperatorDmlElementWiseUnary()
+	: Alpha(1.67326319217681884765625f), Beta(0.0f), Gamma(1.05070102214813232421875f), Num(1)
+{
+}
+
+template<> FOperatorDmlElementWiseUnary<DML_ACTIVATION_ELU_OPERATOR_DESC, NNECore::Internal::EElementWiseUnaryOperatorType::Elu>::FOperatorDmlElementWiseUnary()
+	: Alpha(1.0f), Beta(0.0f), Gamma(1.05070102214813232421875f), Num(1)
+{
+}
+
+template<> FOperatorDmlElementWiseUnary<DML_ACTIVATION_HARD_SIGMOID_OPERATOR_DESC, NNECore::Internal::EElementWiseUnaryOperatorType::HardSigmoid>::FOperatorDmlElementWiseUnary()
+	: Alpha(0.2f), Beta(0.5f), Gamma(0.0f), Num(1)
+{
+}
+
+template<> FOperatorDmlElementWiseUnary<DML_ACTIVATION_LEAKY_RELU_OPERATOR_DESC, NNECore::Internal::EElementWiseUnaryOperatorType::LeakyRelu>::FOperatorDmlElementWiseUnary()
+	: Alpha(0.01f), Beta(0.0f), Gamma(0.0f), Num(1)
+{
+}
+
+/**
+ * Element-wise binary ML operator implementation
+ */
+template
+<
+	typename TDmlElementWiseOpDescType,
+	NNECore::Internal::EElementWiseBinaryOperatorType OpType
+>
+class FOperatorDmlElementWiseBinary : public FOperatorDml
+{
+
+public:
+
+	static FOperatorDml* Create()
+	{
+		return new FOperatorDmlElementWiseBinary();
+	}
+
+private:
+
+	FOperatorDmlElementWiseBinary() = default;
+
+	uint32 Num { 1 };
+
+public:
+
+	//
+	//
+	//
+	virtual bool Initialize(FDeviceContextDml* InDevCtx, TArrayView<const NNECore::Internal::FTensor> InputTensors, TArrayView<const NNECore::Internal::FTensor> OutputTensors, const NNECore::FAttributeMap& Attributes) override
+	{
+		// TODO: Setup attributes
+		Num = OutputTensors[0].GetVolume();
+
+		DevCtx = InDevCtx;
+
+		const NNECore::Internal::FTensor& InputATensorDesc = InputTensors[0];
+		const NNECore::Internal::FTensor& InputBTensorDesc = InputTensors[1];
+		const NNECore::Internal::FTensor& OutputTensorDesc = OutputTensors[0];
+
+		// Initialize tensor descriptors
+		DmlUtil::FTensorDesc	DmlInputATensorDesc{};
+		DmlUtil::FTensorDesc	DmlInputBTensorDesc{};
+		DmlUtil::FTensorDesc	DmlOutputTensorDesc{};
+
+		if (!InitDmlTensorDesc(DmlInputATensorDesc, InputATensorDesc, OutputTensorDesc))
+		{
+			UE_LOG(LogNNX, Warning, TEXT("Failed to initialize tensor(s) for DML inference"));
+			return false;
+		}
+
+		if (!InitDmlTensorDesc(DmlInputBTensorDesc, InputBTensorDesc, OutputTensorDesc))
+		{
+			UE_LOG(LogNNX, Warning, TEXT("Failed to initialize tensor(s) for DML inference"));
+			return false;
+		}
+
+		if (!InitDmlTensorDesc(DmlOutputTensorDesc, OutputTensorDesc))
+		{
+			UE_LOG(LogNNX, Warning, TEXT("Failed to initialize tensor(s) for DML inference"));
+			return false;
+		}
+
+		TDmlElementWiseOpDescType	DmlElemWiseOpDesc{};
+
+		InitDmlOpDesc(DmlElemWiseOpDesc, DmlInputATensorDesc, DmlInputBTensorDesc, DmlOutputTensorDesc);
+
+		DML_OPERATOR_DESC DmlOpDesc{};
+
+		DmlOpDesc.Type = GetDmlElementWiseBinaryOpType<TDmlElementWiseOpDescType>();
+		DmlOpDesc.Desc = &DmlElemWiseOpDesc;
+
+		return CreateOperator(DmlOpDesc);
+	}
+
+private:
+
+	template<typename OpDesc>
+	void InitDmlOpDesc(OpDesc& Desc, DmlUtil::FTensorDesc& LHSTensor, DmlUtil::FTensorDesc& RHSTensor, DmlUtil::FTensorDesc& OutputTensor)
+	{
+		Desc.ATensor = &LHSTensor.Desc;
+		Desc.BTensor = &RHSTensor.Desc;
+		Desc.OutputTensor = &OutputTensor.Desc;
+	}
+
+	void InitDmlOpDesc(DML_ELEMENT_WISE_POW_OPERATOR_DESC& Desc, DmlUtil::FTensorDesc& LHSTensor, DmlUtil::FTensorDesc& RHSTensor, DmlUtil::FTensorDesc& OutputTensor)
+	{
+		Desc.InputTensor = &LHSTensor.Desc;
+		Desc.ExponentTensor = &RHSTensor.Desc;
+		Desc.OutputTensor = &OutputTensor.Desc;
+	}
+
+	void InitDmlOpDesc(DML_ACTIVATION_PARAMETERIZED_RELU_OPERATOR_DESC& Desc, DmlUtil::FTensorDesc& LHSTensor, DmlUtil::FTensorDesc& RHSTensor, DmlUtil::FTensorDesc& OutputTensor)
+	{
+		Desc.InputTensor = &LHSTensor.Desc;
+		Desc.SlopeTensor = &RHSTensor.Desc;
+		Desc.OutputTensor = &OutputTensor.Desc;
+	}
+};
+
+/**
+ * Gemm
+ */
+class FOperatorDmlGemm : public FOperatorDml
+{
+public:
+
+	static FOperatorDml* Create()
+	{
+		return new FOperatorDmlGemm();
+	}
+
+	//
+	//
+	//
+	virtual bool Initialize(FDeviceContextDml* InDevCtx, TArrayView<const NNECore::Internal::FTensor> InputTensors, TArrayView<const NNECore::Internal::FTensor> OutputTensors, const NNECore::FAttributeMap& Attributes) override
+	{
+		// Setup attributes
+		float Alpha = 1.0f;
+		float Beta = 1.0f;
+		int32 TransA = 0;
+		int32 TransB = 0;
+		
+		Alpha = Attributes.GetValueOrDefault(TEXT("alpha"), Alpha);
+		Beta = Attributes.GetValueOrDefault(TEXT("beta"), Beta);
+		TransA = Attributes.GetValueOrDefault(TEXT("transA"), TransA);
+		TransB = Attributes.GetValueOrDefault(TEXT("transB"), TransB);
+		
+		DevCtx = InDevCtx;
+
+		const NNECore::Internal::FTensor& InputATensorDesc = InputTensors[0];
+		const NNECore::Internal::FTensor& InputBTensorDesc = InputTensors[1];
+		const NNECore::Internal::FTensor& OutputTensorDesc = OutputTensors[0];
+
+		// Initialize tensor descriptors
+		DmlUtil::FTensorDesc	DmlInputATensorDesc{};
+		DmlUtil::FTensorDesc	DmlInputBTensorDesc{};
+		DmlUtil::FTensorDesc	DmlInputCTensorDesc{};
+		DmlUtil::FTensorDesc	DmlOutputTensorDesc{};
+
+		if (!InitDmlTensorDesc(DmlInputATensorDesc, InputATensorDesc))
+		{
+			UE_LOG(LogNNX, Warning, TEXT("Failed to initialize tensor(s) for DML inference"));
+			return false;
+		}
+
+		if (!InitDmlTensorDesc(DmlInputBTensorDesc, InputBTensorDesc))
+		{
+			UE_LOG(LogNNX, Warning, TEXT("Failed to initialize tensor(s) for DML inference"));
+			return false;
+		}
+
+		if (InputTensors.Num() > 2)
+		{
+			const NNECore::Internal::FTensor& InputCTensorDesc = InputTensors[2];
+
+			if (!InitDmlTensorDesc(DmlInputCTensorDesc, InputCTensorDesc, OutputTensorDesc))
+			{
+				UE_LOG(LogNNX, Warning, TEXT("Failed to initialize tensor(s) for DML inference"));
+				return false;
+			}
+		}
+
+		if (!InitDmlTensorDesc(DmlOutputTensorDesc, OutputTensorDesc))
+		{
+			UE_LOG(LogNNX, Warning, TEXT("Failed to initialize tensor(s) for DML inference"));
+			return false;
+		}
+
+		DML_GEMM_OPERATOR_DESC	DmlGemmOpDesc{};
+
+		DmlGemmOpDesc.ATensor = &DmlInputATensorDesc.Desc;
+		DmlGemmOpDesc.BTensor = &DmlInputBTensorDesc.Desc;
+		DmlGemmOpDesc.CTensor = InputTensors.Num() > 2 ? &DmlInputCTensorDesc.Desc : nullptr;
+		DmlGemmOpDesc.OutputTensor = &DmlOutputTensorDesc.Desc;
+		DmlGemmOpDesc.Alpha = Alpha;
+		DmlGemmOpDesc.Beta = Beta;
+		DmlGemmOpDesc.TransA = TransA ? DML_MATRIX_TRANSFORM_TRANSPOSE : DML_MATRIX_TRANSFORM_NONE;
+		DmlGemmOpDesc.TransB = TransB ? DML_MATRIX_TRANSFORM_TRANSPOSE : DML_MATRIX_TRANSFORM_NONE;
+
+		DML_OPERATOR_DESC DmlOpDesc{};
+
+		DmlOpDesc.Type = DML_OPERATOR_GEMM;
+		DmlOpDesc.Desc = &DmlGemmOpDesc;
+
+		return CreateOperator(DmlOpDesc);
+	}
+
+};
+
+//
+//
+//
+class FInferenceModelDml : public FInferenceModelRDG
+{
+	class FGraphBuilder;
+	class FBindingTable;
+	class FDebugName;
+
+public:
+
+	FInferenceModelDml();
+	~FInferenceModelDml();
+
+	bool Init(TConstArrayView<uint8> ModelData, FDeviceContextDml* InDevCtx);
+
+protected:
+
+	virtual void AddDispatchOps_RenderThread(FRDGBuilder& GraphBuilder) override;
+	virtual int PrepareTensorShapesAndData() override;
+
+private:
+
+	bool InitCompiledOp(TConstArrayView<int32> OpInputIndices, uint64 TensorDataSize);
+
+	FOperatorDml* OpCreate(const FString& Name, TArrayView<const NNECore::Internal::FTensor> InputTensorDesc, TArrayView<const NNECore::Internal::FTensor> OutputTensorDescs, const NNECore::FAttributeMap& Attributes);
+		
+	FBufferRHIRef CreateRHIBuffer(FRHICommandListImmediate& RHICmdList, uint32 Size, EBufferUsageFlags Usage, ERHIAccess Access, const TCHAR* DbgName);
+	ID3D12Resource* CreateD3D12Buffer(uint32 Size, D3D12_RESOURCE_STATES ResourceState = D3D12_RESOURCE_STATE_COMMON, D3D12_HEAP_TYPE HeapType = D3D12_HEAP_TYPE_DEFAULT, const TCHAR* DebugName = nullptr);
+	
+	// TODO: FIXME: This should go into RDG
+	static constexpr int32 MaxNumInputs = 32;
+	static constexpr int32 MaxNumOutputs = 4;
+
+	using FRHIBufferInputArray = TArray<FRHIBuffer*, TInlineAllocator<MaxNumInputs>>;
+	using FRHIBufferOutputArray = TArray<FRHIBuffer*, TInlineAllocator<MaxNumOutputs>>;
+
+	TComPtr<IDMLOperatorInitializer>	OpInit;
+	TComPtr<IDMLCompiledOperator>		CompiledOp;
+	FDeviceContextDml*					DevCtx;
+	TUniquePtr<FBindingTable>			BindingTable;
+	TComPtr<ID3D12DescriptorHeap>		DescHeap;
+	uint32								DescCount;
+	uint32								DescSize;
+
+	FRHIBufferInputArray				InputBuffers;
+	FRHIBufferOutputArray				OutputBuffers;
+#ifdef NNE_USE_D3D12_RESOURCES
+	TComPtr<ID3D12Resource>				PersistBuff;
+	TComPtr<ID3D12Resource>				TempBuff;
+#else
+	FBufferRHIRef						PersistBuff;
+#endif
+	uint64								MemSizeTemp;
+	uint64								MemSizePersist;
+	ID3D12DynamicRHI*					DynamicRHI;
+};
+
+//
+//
+//
+class FRuntimeDml : public FRuntimeRDG
+{
+public:
+
+	FRuntimeDml() = default;
+	virtual ~FRuntimeDml();
+
+	virtual FString GetRuntimeName() const override
+	{
+		return NNX_RUNTIME_DML_NAME;
+	}
+	
+	virtual EMLRuntimeSupportFlags GetSupportFlags() const override
+	{
+		return EMLRuntimeSupportFlags::RDG;
+	}
+
+	virtual TArray<uint8> CreateModelData(FString FileType, TConstArrayView<uint8> FileData) override
+	{
+		if (!CanCreateModelData(FileType, FileData))
+		{
+			return {};
+		}
+
+		TUniquePtr<IModelOptimizer> Optimizer = NNEUtils::Internal::CreateONNXToNNEModelOptimizer();
+
+		FNNIModelRaw InputModel;
+		InputModel.Data = FileData;
+		InputModel.Format = ENNXInferenceFormat::ONNX;
+
+		FNNIModelRaw OutputModel;
+		if (!Optimizer->Optimize(InputModel, OutputModel, {}))
+		{
+			return {};
+		}
+
+		return ConvertToModelData(OutputModel.Data);
+	};
+
+	virtual TUniquePtr<FMLInferenceModel> CreateModel(TConstArrayView<uint8> ModelData) override
+	{
+		if (!CanCreateModel(ModelData))
+		{
+			return TUniquePtr<FMLInferenceModel>();
+		}
+
+		// Create the model and initialize it with the data not including the header
+		FInferenceModelDml* Model = new FInferenceModelDml();
+		if (!Model->Init(ModelData, &Ctx))
+		{
+			delete Model;
+			return TUniquePtr<FMLInferenceModel>();
+		}
+		return TUniquePtr<FMLInferenceModel>(Model);
+	}
+
+	bool Init(bool bRegisterOnlyOperators);
+
+private:
+
+	bool RegisterElementWiseUnaryOperators();
+	bool RegisterElementWiseBinaryOperators();
+	bool RegisterGemmOperator();
+	
+	FDeviceContextDml		Ctx;
+};
+
+//
+//
+//
+static TUniquePtr<FRuntimeDml> GDmlRuntime;
+
+//
+//
+//
+FRuntimeDml::~FRuntimeDml()
+{
+}
+
+//
+//
+//
+bool FRuntimeDml::Init(bool bRegisterOnlyOperators)
+{
+	RegisterElementWiseUnaryOperators();
+	RegisterElementWiseBinaryOperators();
+	RegisterGemmOperator();
+
+	if (bRegisterOnlyOperators)
+	{
+		UE_LOG(LogNNX, Display, TEXT("Registering only operators"));
+		return true;
+	}
+
+	HRESULT Res;
+
+	// In order to use DirectML we need D3D12
+	ID3D12DynamicRHI* RHI = nullptr;
+
+	if (GDynamicRHI && GDynamicRHI->GetInterfaceType() == ERHIInterfaceType::D3D12)
+	{
+		RHI = GetID3D12PlatformDynamicRHI();
+
+		if (!RHI)
+		{
+			UE_LOG(LogNNX, Warning, TEXT("Error:%s RHI is not supported by DirectML"), GDynamicRHI->GetName());
+			return false;
+		}
+	}
+	else
+	{
+		if (GDynamicRHI)
+		{
+			UE_LOG(LogNNX, Warning, TEXT("Error:%s RHI is not supported by DirectML"), GDynamicRHI->GetName());
+			return false;
+		}
+		else
+		{
+			UE_LOG(LogNNX, Warning, TEXT("Error:No RHI found"));
+			return false;
+		}
+	}
+
+	Ctx.DeviceIndex = 0;
+	Ctx.D3D12Device = RHI->RHIGetDevice(Ctx.DeviceIndex);
+
+#if PLATFORM_WINDOWS
+	ID3D12Device5* D3D12Device5 = nullptr;
+	
+	Res = Ctx.D3D12Device->QueryInterface(&D3D12Device5);
+	if (D3D12Device5)
+	{
+		uint32	NumCommands = 0;
+
+		Res = D3D12Device5->EnumerateMetaCommands(&NumCommands, nullptr);
+		if (NumCommands)
+		{
+			UE_LOG(LogNNX, Verbose, TEXT("D3D12 Meta commands:%u"), NumCommands);
+
+			TArray<D3D12_META_COMMAND_DESC>	MetaCmds;
+
+			MetaCmds.SetNumUninitialized(NumCommands);
+
+			Res = D3D12Device5->EnumerateMetaCommands(&NumCommands, MetaCmds.GetData());
+			for (uint32 Idx = 0; Idx < NumCommands; ++Idx)
+			{
+				const D3D12_META_COMMAND_DESC& Desc = MetaCmds[Idx];
+
+				UE_LOG(LogNNX, Verbose, TEXT("   %s"), Desc.Name);
+			}
+		}
+	}
+#endif
+
+	DML_CREATE_DEVICE_FLAGS DmlCreateFlags = DML_CREATE_DEVICE_FLAG_NONE;
+
+	// Set debugging flags
+	if (RHI->IsD3DDebugEnabled())
+	{
+		DmlCreateFlags |= DML_CREATE_DEVICE_FLAG_DEBUG;
+	}
+
+	Res = DMLCreateDevice(Ctx.D3D12Device, DmlCreateFlags, DML_PPV_ARGS(&Ctx.Device));
+	if (!Ctx.Device)
+	{
+		UE_LOG(LogNNX, Warning, TEXT("Failed to create DirectML device, res:%x"), Res);
+		return false;
+	}
+
+	DML_FEATURE_QUERY_TENSOR_DATA_TYPE_SUPPORT	Fp16Query = { DML_TENSOR_DATA_TYPE_FLOAT16 };
+	DML_FEATURE_DATA_TENSOR_DATA_TYPE_SUPPORT	Fp16Supported = {};
+	
+	Ctx.Device->CheckFeatureSupport(DML_FEATURE_TENSOR_DATA_TYPE_SUPPORT, sizeof(Fp16Query), &Fp16Query, sizeof(Fp16Supported), &Fp16Supported);
+
+	DML_FEATURE_LEVEL					FeatureLevels[] = { DML_FEATURE_LEVEL_5_0 };
+	DML_FEATURE_QUERY_FEATURE_LEVELS	FeatureLevelQuery = { UE_ARRAY_COUNT(FeatureLevels), FeatureLevels };
+	DML_FEATURE_DATA_FEATURE_LEVELS		FeatureLevelSupported = {};
+
+	Res = Ctx.Device->CheckFeatureSupport(DML_FEATURE_FEATURE_LEVELS, sizeof(FeatureLevelQuery), &FeatureLevelQuery, sizeof(FeatureLevelSupported), &FeatureLevelSupported);
+	if (FAILED(Res) || FeatureLevelSupported.MaxSupportedFeatureLevel < DML_FEATURE_LEVEL_5_0)
+	{
+		UE_LOG(LogNNX, Warning, TEXT("DirectML feature level %x not supported"), FeatureLevels[0]);
+		return false;
+	}
+
+	Res = Ctx.Device->CreateCommandRecorder(DML_PPV_ARGS(&Ctx.CmdRec));
+	if (!Ctx.CmdRec)
+	{
+		UE_LOG(LogNNX, Warning, TEXT("Failed to create DML command recorder, res:%x"), Res);
+		return false;
+	}
+
+	return true;
+}
+
+//
+//
+//
+bool FRuntimeDml::RegisterElementWiseUnaryOperators()
+{
+#define OP(DmlOpDesc, OpName) FOperatorRegistryDml::Get()->OpAdd(TEXT(#OpName), FOperatorDmlElementWiseUnary<DmlOpDesc, NNECore::Internal::EElementWiseUnaryOperatorType::OpName>::Create)
+
+	OP(DML_ELEMENT_WISE_ABS_OPERATOR_DESC, Abs);
+	OP(DML_ELEMENT_WISE_ACOS_OPERATOR_DESC, Acos);
+	OP(DML_ELEMENT_WISE_ACOSH_OPERATOR_DESC, Acosh);
+	OP(DML_ELEMENT_WISE_ASIN_OPERATOR_DESC, Asin);
+	OP(DML_ELEMENT_WISE_ASINH_OPERATOR_DESC, Asinh);
+	OP(DML_ELEMENT_WISE_ATAN_OPERATOR_DESC, Atan);
+	OP(DML_ELEMENT_WISE_ATANH_OPERATOR_DESC, Atanh);
+	OP(DML_ELEMENT_WISE_CEIL_OPERATOR_DESC, Ceil);
+	OP(DML_ELEMENT_WISE_COS_OPERATOR_DESC, Cos);
+	OP(DML_ELEMENT_WISE_COSH_OPERATOR_DESC, Cosh);
+	OP(DML_ACTIVATION_ELU_OPERATOR_DESC, Elu);
+	OP(DML_ELEMENT_WISE_ERF_OPERATOR_DESC, Erf);
+	OP(DML_ELEMENT_WISE_EXP_OPERATOR_DESC, Exp);
+	OP(DML_ELEMENT_WISE_FLOOR_OPERATOR_DESC, Floor);
+	OP(DML_ELEMENT_WISE_IS_INFINITY_OPERATOR_DESC, IsInf);
+	OP(DML_ELEMENT_WISE_IS_NAN_OPERATOR_DESC, IsNan);
+	OP(DML_ACTIVATION_HARD_SIGMOID_OPERATOR_DESC, HardSigmoid);
+	//OP(HardSwish);
+	OP(DML_ACTIVATION_LEAKY_RELU_OPERATOR_DESC, LeakyRelu);
+	OP(DML_ELEMENT_WISE_LOG_OPERATOR_DESC, Log);
+	OP(DML_ELEMENT_WISE_NEGATE_OPERATOR_DESC, Neg);
+	//OP(Not);
+	OP(DML_ELEMENT_WISE_RECIP_OPERATOR_DESC, Reciprocal);
+	OP(DML_ACTIVATION_RELU_OPERATOR_DESC, Relu);
+	OP(DML_ELEMENT_WISE_ROUND_OPERATOR_DESC, Round);
+	OP(DML_ACTIVATION_SCALED_ELU_OPERATOR_DESC, Selu);
+	OP(DML_ACTIVATION_SIGMOID_OPERATOR_DESC, Sigmoid);
+	OP(DML_ELEMENT_WISE_SIGN_OPERATOR_DESC, Sign);
+	OP(DML_ELEMENT_WISE_SIN_OPERATOR_DESC, Sin);
+	OP(DML_ELEMENT_WISE_SINH_OPERATOR_DESC, Sinh);
+	OP(DML_ACTIVATION_SOFTPLUS_OPERATOR_DESC, Softplus);
+	OP(DML_ACTIVATION_SOFTSIGN_OPERATOR_DESC, Softsign);
+	OP(DML_ELEMENT_WISE_SQRT_OPERATOR_DESC, Sqrt);
+	OP(DML_ELEMENT_WISE_TAN_OPERATOR_DESC, Tan);
+	OP(DML_ELEMENT_WISE_TANH_OPERATOR_DESC, Tanh);
+
+#undef OP
+
+	return true;
+}
+
+//
+//
+//
+bool FRuntimeDml::RegisterElementWiseBinaryOperators()
+{
+#define OP(DmlOpDesc, OpName) FOperatorRegistryDml::Get()->OpAdd(TEXT(#OpName), FOperatorDmlElementWiseBinary<DmlOpDesc, NNECore::Internal::EElementWiseBinaryOperatorType::OpName>::Create)
+
+	OP(DML_ELEMENT_WISE_ADD_OPERATOR_DESC, Add);
+	// And
+	OP(DML_ELEMENT_WISE_DIVIDE_OPERATOR_DESC, Div);
+	// Equal
+	// Greater
+	// GreaterOrEqual
+	// Less
+	// LessOrEqual
+	// Mod
+	OP(DML_ELEMENT_WISE_MULTIPLY_OPERATOR_DESC, Mul);
+	// Or
+	OP(DML_ACTIVATION_PARAMETERIZED_RELU_OPERATOR_DESC, Prelu);
+	OP(DML_ELEMENT_WISE_POW_OPERATOR_DESC, Pow);
+	OP(DML_ELEMENT_WISE_SUBTRACT_OPERATOR_DESC, Sub);
+	// Xor
+
+#undef OP
+
+	return true;
+}
+
+//
+//
+//
+bool FRuntimeDml::RegisterGemmOperator()
+{
+	FOperatorRegistryDml::Get()->OpAdd(TEXT("Gemm"), FOperatorDmlGemm::Create);
+	return true;
+}
+
+//
+//
+//
+class FInferenceModelDml::FDebugName
+{
+	static constexpr int32 Size = 128;
+
+public:
+
+	FDebugName()
+	{
+		Str[0] = '\0';
+		Length = 0;
+	}
+
+	FDebugName(const FString& InStr)
+	{
+		FTCHARToUTF8 Conv(*InStr);
+
+		Length = Conv.Length() + 1 < Size ? Conv.Length() + 1 : Size - 1;
+		FCStringAnsi::Strncpy(Str, Conv.Get(), Length);
+		Str[Length] = '\0';
+	}
+
+	FDebugName(FStringView InStr)
+	{
+		FTCHARToUTF8 Conv(InStr.GetData());
+
+		Length = Conv.Length() + 1 < Size ? Conv.Length() + 1 : Size - 1;
+		FCStringAnsi::Strncpy(Str, Conv.Get(), Length);
+		Str[Length] = '\0';
+	}
+	~FDebugName()
+	{
+	}
+
+	const char* Get() const
+	{
+		return Str;
+	}
+
+private:
+
+	char	Str[Size];
+	int32	Length;
+};
+
+//
+//
+//
+class FInferenceModelDml::FBindingTable
+{
+public:
+
+	bool Init(FInferenceModelDml* InModel)
+	{
+		Model = InModel;
+		DynamicRHI = InModel->DynamicRHI;
+
+		return true;
+	}
+
+#ifdef NNE_USE_D3D12_RESOURCES
+	void Bind(IDMLOperatorInitializer* OpInit, TConstArrayView<FRHIBuffer*> InputBuffers, ID3D12Resource* PersistBuff, ID3D12Resource* TempBuff = nullptr)
+#else
+	void Bind(IDMLOperatorInitializer* OpInit, TConstArrayView<FRHIBuffer*> InputBuffers, FRHIBuffer* PersistBuff, FRHIBuffer* TempBuff = nullptr)
+#endif
+	{
+		Reset(OpInit);
+
+		TArray<DML_BUFFER_BINDING, TInlineAllocator<MaxNumInputs>> Inputs;
+
+		for (FRHIBuffer* Buffer : InputBuffers)
+		{
+			if (Buffer)
+			{
+				Inputs.Emplace(MakeBind(Buffer));
+			}
+			else
+			{
+				Inputs.Add({});
+			}
+		}
+
+		DML_BUFFER_ARRAY_BINDING	InputBindArray{ Inputs.Num(), Inputs.GetData() };
+		DML_BINDING_DESC			InputBindArrayDesc{ DML_BINDING_TYPE_BUFFER_ARRAY, &InputBindArray };
+		
+		BindingTable->BindInputs(1, &InputBindArrayDesc);
+		
+		DML_BUFFER_BINDING			PersistBind {};
+		DML_BINDING_DESC			PersistBindDesc{ DML_BINDING_TYPE_BUFFER, &PersistBind };
+
+		if (PersistBuff)
+		{
+#ifdef NNE_USE_D3D12_RESOURCES
+			PersistBind = DML_BUFFER_BINDING { PersistBuff, 0, PersistBuff->GetDesc().Width };
+#else
+			PersistBind = MakeBind(PersistBuff);
+#endif
+		}
+
+		BindingTable->BindOutputs(1, &PersistBindDesc);
+
+		DML_BUFFER_BINDING			TempBind{};
+		DML_BINDING_DESC			TempBindDesc{ DML_BINDING_TYPE_BUFFER, &TempBind };
+
+		if (TempBuff)
+		{
+#ifdef NNE_USE_D3D12_RESOURCES
+			TempBind = { PersistBuff, 0, PersistBuff->GetDesc().Width };
+#else
+			TempBind = MakeBind(TempBuff);
+#endif
+			BindingTable->BindTemporaryResource(&TempBindDesc);
+		}
+	}
+
+#ifdef NNE_USE_D3D12_RESOURCES
+	void Bind(IDMLCompiledOperator* Op, TConstArrayView<FRHIBuffer*> InputBuffers, TConstArrayView<FRHIBuffer*> OutputBuffers, ID3D12Resource* PersistBuff = nullptr, ID3D12Resource* TempBuff = nullptr)
+#else
+	
+	void Bind(IDMLCompiledOperator* Op, TConstArrayView<FRHIBuffer*> InputBuffers, TConstArrayView<FRHIBuffer*> OutputBuffers, FRHIBuffer* PersistBuff = nullptr, FRHIBuffer* TempBuff = nullptr)
+#endif
+	{
+		Reset(Op);
+
+		for (FRHIBuffer* Buffer : InputBuffers)
+		{
+			AddBind(Buffer, InputBinds, InputBindDescs);
+		}
+
+		for (FRHIBuffer* Buffer : OutputBuffers)
+		{
+			AddBind(Buffer, OutputBinds, OutputBindDescs);
+		}
+
+		BindingTable->BindInputs(InputBinds.Num(), InputBindDescs.GetData());
+		BindingTable->BindOutputs(OutputBinds.Num(), OutputBindDescs.GetData());
+
+		DML_BUFFER_BINDING			PersistBind;
+		DML_BINDING_DESC			PersistBindDesc{ DML_BINDING_TYPE_BUFFER, &PersistBind };
+
+		if (PersistBuff)
+		{
+#ifdef NNE_USE_D3D12_RESOURCES
+			PersistBind = { PersistBuff, 0, PersistBuff->GetDesc().Width };
+#else
+			PersistBind =  MakeBind(PersistBuff);MakeBind(PersistBuff);
+#endif
+			BindingTable->BindPersistentResource(&PersistBindDesc);
+		}
+
+		DML_BUFFER_BINDING			TempBind{};
+		DML_BINDING_DESC			TempBindDesc{ DML_BINDING_TYPE_BUFFER, &TempBind };
+
+		if (TempBuff)
+		{
+			
+#ifdef NNE_USE_D3D12_RESOURCES
+			TempBind = { TempBuff, 0, TempBuff->GetDesc().Width };
+#else
+			TempBind = MakeBind(TempBuff);
+#endif
+			BindingTable->BindTemporaryResource(&TempBindDesc);
+		}
+	}
+
+
+	IDMLBindingTable* Get()
+	{
+		return BindingTable;
+	}
+
+private:
+
+	bool Reset(IDMLDispatchable* Disp)
+	{
+		InputBinds.Reset();
+		InputBindDescs.Reset();
+		OutputBinds.Reset();
+		OutputBindDescs.Reset();
+		
+		DML_BINDING_PROPERTIES BindingProps = Disp->GetBindingProperties();
+		DML_BINDING_TABLE_DESC Desc{};
+
+		Desc.Dispatchable = Disp;
+		Desc.CPUDescriptorHandle = CD3DX12_CPU_DESCRIPTOR_HANDLE(Model->DescHeap->GetCPUDescriptorHandleForHeapStart(), 0, Model->DescSize);
+		Desc.GPUDescriptorHandle = CD3DX12_GPU_DESCRIPTOR_HANDLE(Model->DescHeap->GetGPUDescriptorHandleForHeapStart(), 0, Model->DescSize);
+		Desc.SizeInDescriptors = Model->DescCount;
+
+		HRESULT Res;
+
+		if (!BindingTable)
+		{
+			Res = Model->DevCtx->Device->CreateBindingTable(&Desc, DML_PPV_ARGS(&BindingTable));
+			if (!BindingTable)
+			{
+				UE_LOG(LogNNX, Warning, TEXT("Failed to create DML binding table, res:%d"), Res);
+				return false;
+			}
+		}
+		else
+		{
+			BindingTable->Reset(&Desc);
+		}
+
+		return true;
+	}
+
+	template<class TBindingArray, class TDescArray>
+	void AddBind(FRHIBuffer* Buffer, TBindingArray& Bindings, TDescArray& Descs)
+	{
+		if (Buffer)
+		{
+			DML_BUFFER_BINDING& Bind = Bindings.Add_GetRef(MakeBind(Buffer));
+			Descs.Add({ DML_BINDING_TYPE_BUFFER, &Bind });
+		}
+		else
+		{
+			DML_BUFFER_BINDING& Bind = Bindings.Add_GetRef({});
+			Descs.Add({ DML_BINDING_TYPE_BUFFER, &Bind });
+		}
+	}
+
+	DML_BUFFER_BINDING MakeBind(FRHIBuffer* Buffer)
+	{
+		ID3D12Resource* Resource = DynamicRHI->RHIGetResource(Buffer);
+
+		return DML_BUFFER_BINDING{ Resource, 0, Buffer->GetSize() };
+	}
+
+	TComPtr<IDMLBindingTable>										BindingTable;
+	TArray<DML_BUFFER_BINDING, TInlineAllocator<MaxNumInputs>>		InputBinds;
+	TArray<DML_BINDING_DESC, TInlineAllocator<MaxNumInputs>>		InputBindDescs;
+	TArray<DML_BUFFER_BINDING, TInlineAllocator<MaxNumOutputs>>		OutputBinds;
+	TArray<DML_BINDING_DESC, TInlineAllocator<MaxNumOutputs>>		OutputBindDescs;
+	ID3D12DynamicRHI*												DynamicRHI;
+	FInferenceModelDml*												Model;
+};
+
+//
+//
+//
+class FInferenceModelDml::FGraphBuilder
+{
+public:
+
+	struct FOpDesc
+	{
+		FOperatorDml*	Op;
+		int32			InputStart;
+		int32			InputCount;
+		int32			OutputStart;
+		int32			OutputCount;
+		FDebugName		DbgName;
+	};
+
+	struct FGraphDesc
+	{
+		TConstArrayView<NNECore::Internal::FTensor>	AllTensors;
+		TConstArrayView<int32>		InputIndices;
+		TConstArrayView<int32>		OutputIndices;
+		TConstArrayView<int32>		WeightIndices;
+		TConstArrayView<int32>		IntermediateIndices;
+		TConstArrayView<FTensorRDG>	WeightTensors;
+		TConstArrayView<FOpDesc>	Operators;
+		TConstArrayView<int32>		OpInputIndices;
+		TConstArrayView<int32>		OpOutputIndices;
+	};
+	
+private:
+
+	enum class EEdgeType
+	{
+		Input,
+		Output,
+		Intermediate
+	};
+
+	struct FEdge
+	{
+		EEdgeType	Type;
+		int32		TensorIdx{ -1 };
+		int32		NodeSrc{ -1 };
+		int32		NodeSrcOutput{ -1 };
+		int32		NodeDst{ -1 };
+		int32		NodeDstInput{ -1 };
+		
+		FEdge(EEdgeType InType)
+			: Type(InType)
+		{
+		}
+
+		FEdge& SetTensorIdx(int32 Value)
+		{
+			TensorIdx = Value;
+			return *this;
+		}
+
+		FEdge& SetNodeSrc(int32 Value)
+		{
+			NodeSrc = Value;
+			return *this;
+		}
+
+		FEdge& SetNodeSrcOutput(int32 Value)
+		{
+			NodeSrcOutput = Value;
+			return *this;
+		}
+
+		FEdge& SetNodeDst(int32 Value)
+		{
+			NodeDst = Value;
+			return *this;
+		}
+
+		FEdge& SetNodeDstInput(int32 Value)
+		{
+			NodeDstInput = Value;
+			return *this;
+		}
+	};
+
+public:
+
+	IDMLCompiledOperator* Compile(FDeviceContextDml* DevCtx, const FGraphDesc& InGraph)
+	{
+		IDMLDevice*				Device = DevCtx->Device;
+		TComPtr<IDMLDevice1>	Device1;
+
+		Device1.FromQueryInterface(__uuidof(IDMLDevice1), DevCtx->Device);
+		check(Device1);
+		if (!Device1)
+		{
+			return nullptr;
+		}
+
+		if (!AddEdges(InGraph))
+		{
+			return nullptr;
+		}
+
+		TArray<DML_INPUT_GRAPH_EDGE_DESC>			InputEdges;
+		TArray<DML_OUTPUT_GRAPH_EDGE_DESC>			OutputEdges;
+		TArray<DML_INTERMEDIATE_GRAPH_EDGE_DESC>	IntermediateEdges;
+		TArray<FDebugName>							DbgInputNames;
+		TArray<FDebugName>							DbgIntermediateNames;
+		TArray<FDebugName>							DbgOutputNames;
+
+		DbgInputNames.Reserve(InGraph.AllTensors.Num());
+		DbgIntermediateNames.Reserve(InGraph.AllTensors.Num());
+
+		for (const FEdge& Edge : Edges)
+		{
+			if (Edge.Type == EEdgeType::Input)
+			{
+				check(Edge.NodeSrcOutput >= 0);
+				check(Edge.NodeDst >= 0);
+				check(Edge.NodeDstInput >= 0);
+
+				DML_INPUT_GRAPH_EDGE_DESC& Input = InputEdges.Add_GetRef({});
+
+				Input.GraphInputIndex = Edge.NodeSrcOutput;
+				Input.ToNodeIndex = Edge.NodeDst;
+				Input.ToNodeInputIndex = Edge.NodeDstInput;
+				
+				const FDebugName& DbgName = DbgInputNames.Add_GetRef(InGraph.AllTensors[Edge.TensorIdx].GetName());
+				Input.Name = DbgName.Get();
+			}
+			else if (Edge.Type == EEdgeType::Output)
+			{
+				check(Edge.NodeDstInput >= 0);
+				check(Edge.NodeSrc >= 0);
+				check(Edge.NodeSrcOutput >= 0);
+
+				DML_OUTPUT_GRAPH_EDGE_DESC&	Output = OutputEdges.Add_GetRef({});
+
+				Output.GraphOutputIndex = Edge.NodeDstInput;
+				Output.FromNodeIndex = Edge.NodeSrc;
+				Output.FromNodeOutputIndex = Edge.NodeSrcOutput;
+
+				const FDebugName& DbgName = DbgOutputNames.Add_GetRef(InGraph.AllTensors[Edge.TensorIdx].GetName());
+				Output.Name = DbgName.Get();
+			}
+			else if (Edge.Type == EEdgeType::Intermediate)
+			{
+				DML_INTERMEDIATE_GRAPH_EDGE_DESC& Intermediate = IntermediateEdges.Add_GetRef({});
+
+				check(Edge.NodeSrc >= 0);
+				check(Edge.NodeSrcOutput >= 0);
+				check(Edge.NodeDst >= 0);
+				check(Edge.NodeDstInput >= 0);
+
+				Intermediate.FromNodeIndex = Edge.NodeSrc;
+				Intermediate.FromNodeOutputIndex = Edge.NodeSrcOutput;
+				Intermediate.ToNodeIndex = Edge.NodeDst;
+				Intermediate.ToNodeInputIndex = Edge.NodeDstInput;
+
+				const FDebugName& DbgName = DbgIntermediateNames.Add_GetRef(InGraph.AllTensors[Edge.TensorIdx].GetName());
+				Intermediate.Name = DbgName.Get();
+			}
+		}
+			
+		TArray<DML_GRAPH_NODE_DESC>		Nodes;
+		TArray<DML_GRAPH_EDGE_DESC>		InputEdgeDescs;
+		TArray<DML_GRAPH_EDGE_DESC>		OutputEdgeDescs;
+		TArray<DML_GRAPH_EDGE_DESC>		IntermediateEdgeDescs;
+
+		Nodes.SetNumUninitialized(Operators.Num());
+
+		for (int32 Idx = 0; Idx < Operators.Num(); ++Idx)
+		{
+			DML_GRAPH_NODE_DESC& NodeDesc = Nodes[Idx];
+
+			NodeDesc.Type = DML_GRAPH_NODE_TYPE_OPERATOR;
+			NodeDesc.Desc = &Operators[Idx];
+		}
+
+		for (int32 Idx = 0; Idx < InputEdges.Num(); ++Idx)
+		{
+			DML_GRAPH_EDGE_DESC& Edge = InputEdgeDescs.Add_GetRef({});
+
+			Edge.Type = DML_GRAPH_EDGE_TYPE_INPUT;
+			Edge.Desc = &InputEdges[Idx];
+		}			
+
+		for (int32 Idx = 0; Idx < OutputEdges.Num(); ++Idx)
+		{
+			DML_GRAPH_EDGE_DESC& Edge = OutputEdgeDescs.Add_GetRef({});
+
+			Edge.Type = DML_GRAPH_EDGE_TYPE_OUTPUT;
+			Edge.Desc = &OutputEdges[Idx];				
+		}
+
+		for (int32 Idx = 0; Idx < IntermediateEdges.Num(); ++Idx)
+		{
+			DML_GRAPH_EDGE_DESC& Edge = IntermediateEdgeDescs.Add_GetRef({});
+
+			Edge.Type = DML_GRAPH_EDGE_TYPE_INTERMEDIATE;
+			Edge.Desc = &IntermediateEdges[Idx];
+		}
+			
+		DML_GRAPH_DESC	Graph = DML_GRAPH_DESC{};
+
+		Graph.InputCount = InputEdges.Num();
+		Graph.OutputCount = OutputEdges.Num();
+		Graph.NodeCount = Operators.Num();
+		Graph.Nodes = Nodes.GetData();
+		Graph.InputEdgeCount = InputEdgeDescs.Num();
+		Graph.InputEdges = InputEdgeDescs.GetData();
+		Graph.OutputEdgeCount = OutputEdgeDescs.Num();
+		Graph.OutputEdges = OutputEdgeDescs.GetData();
+		Graph.IntermediateEdgeCount = IntermediateEdgeDescs.Num();
+		Graph.IntermediateEdges = IntermediateEdgeDescs.GetData();
+
+		IDMLCompiledOperator* Op = nullptr;
+		HRESULT Res;
+			
+		Res = Device1->CompileGraph(&Graph, DML_EXECUTION_FLAG_NONE, DML_PPV_ARGS(&Op));
+		if (FAILED(Res))
+		{
+			UE_LOG(LogNNX, Warning, TEXT("Error:Failed to compile DML graph"));
+			Op = nullptr;
+		};
+
+		return Op;
+	}
+
+private:
+
+	bool AddEdges(const FGraphDesc& InGraph)
+	{
+		Edges.Reset();
+		Operators.Reset();
+		NumInputs = 0;
+		NumOutputs = 0;
+
+		for (int32 Idx = 0; Idx < InGraph.InputIndices.Num(); ++Idx)
+		{
+			const int32 TensorIdx = InGraph.InputIndices[Idx];
+			AddInput(TensorIdx);
+		}
+
+		for (int32 Idx = 0; Idx < InGraph.WeightIndices.Num(); ++Idx)
+		{
+			const int32 TensorIdx = InGraph.WeightIndices[Idx];
+			AddInput(TensorIdx);
+		}
+
+		//const int32 OutputIndexStart = InGraph.InputIndices.Num() + InGraph.WeightIndices.Num() + InGraph.IntermediateIndices.Num();
+
+		for (int32 Idx = 0; Idx < InGraph.OutputIndices.Num(); ++Idx)
+		{
+			const int32 TensorIdx = InGraph.OutputIndices[Idx];
+			AddOutput(TensorIdx);
+		}
+
+		Operators.Reset(InGraph.Operators.Num());
+		for (const FOpDesc& OpDesc : InGraph.Operators)
+		{
+			AddOp(OpDesc, InGraph);
+		}
+
+		// TODO: FIXME: Validate edges here
+		
+		return true;
+	}
+
+	void AddInput(int32 TensorIdx)
+	{
+		AddEdge(
+			FEdge(EEdgeType::Input)
+				.SetTensorIdx(TensorIdx)
+				.SetNodeSrcOutput(NumInputs)
+		);
+
+		++NumInputs;
+	}
+
+	void AddOutput(int32 TensorIdx)
+	{
+		AddEdge(
+			FEdge(EEdgeType::Output)
+				.SetTensorIdx(TensorIdx)
+				.SetNodeDstInput(NumOutputs)
+		);
+
+		++NumOutputs;
+	}
+
+	void AddIntermediate(int32 TensorIdx, int32 NodeSrc, int32 NodeSrcOutput)
+	{
+		FEdge* ConnEdge = Edges.FindByPredicate(
+				[TensorIdx](const FEdge& Curr)
+				{
+					return Curr.TensorIdx == TensorIdx;
+				}
+		);
+
+		if (ConnEdge)
+		{
+			ConnectEdgeSrc(TensorIdx, NodeSrc, NodeSrcOutput);
+		}
+		else
+		{
+			AddEdge(
+				FEdge(EEdgeType::Intermediate)
+					.SetTensorIdx(TensorIdx)
+					.SetNodeSrc(NodeSrc)
+					.SetNodeSrcOutput(NodeSrcOutput)
+			);
+		}
+	}
+
+	void AddEdge(const FEdge& Edge)
+	{
+		FEdge* StartEdge = Edges.FindByPredicate(
+				[Edge](const FEdge& Curr)
+				{
+					return Curr.TensorIdx == Edge.TensorIdx;
+				}
+		);
+
+		check(StartEdge == nullptr);
+		Edges.Add(Edge);
+	}
+
+	bool ConnectEdgeDst(int32 TensorIdx, int32 NodeDst, int32 NodeDstInput)
+	{
+		FEdge* StartEdge = Edges.FindByPredicate(
+				[TensorIdx](const FEdge& Curr)
+				{
+					return Curr.TensorIdx == TensorIdx;
+				}
+		);
+
+		bool bFoundEdge = false;
+
+		for (FEdge* Curr = StartEdge; Curr->TensorIdx == TensorIdx; ++Curr)
+		{
+			if (Curr->NodeDst == -1 && Curr->NodeDstInput == -1)
+			{
+				Curr->NodeDst = NodeDst;
+				Curr->NodeDstInput = NodeDstInput;
+				bFoundEdge = true;
+				break;
+			}
+			else if (Curr->NodeDst == NodeDst && Curr->NodeDstInput == NodeDstInput)
+			{
+				bFoundEdge = true;
+				break;
+			}
+		}
+
+		checkf(bFoundEdge, TEXT("ConnectEdgeDst() has failed"));
+		return bFoundEdge;
+	}
+
+	bool ConnectEdgeSrc(int32 TensorIdx, int32 NodeSrc, int32 NodeSrcOutput)
+	{
+		FEdge* StartEdge =
+			Edges.FindByPredicate(
+				[TensorIdx](const FEdge& Curr)
+				{
+					return Curr.TensorIdx == TensorIdx;
+				}
+		);
+
+		bool bFoundEdge = false;
+
+		for (FEdge* Curr = StartEdge; Curr->TensorIdx == TensorIdx; ++Curr)
+		{
+			if (Curr->NodeSrc == -1 && Curr->NodeSrcOutput == -1)
+			{
+				Curr->NodeSrc = NodeSrc;
+				Curr->NodeSrcOutput = NodeSrcOutput;
+				bFoundEdge = true;
+				break;
+			}
+			else if (Curr->NodeSrc == NodeSrc && Curr->NodeSrcOutput == NodeSrcOutput)
+			{
+				bFoundEdge = true;
+				break;
+			}
+		}
+
+		checkf(bFoundEdge, TEXT("ConnectEdgeSrc() has failed"));
+		return bFoundEdge;
+	}
+
+	void AddOp(const FOpDesc& InOp, const FGraphDesc& InGraph)
+	{
+		DML_OPERATOR_GRAPH_NODE_DESC& OpDesc = Operators.Add_GetRef({});
+
+		OpDesc.Operator = InOp.Op->GetOperator();
+		OpDesc.Name = InOp.DbgName.Get();
+
+		const int32 NodeIdx = Operators.Num() - 1;
+
+		for (int32 Idx = 0; Idx < InOp.InputCount; ++Idx)
+		{
+			const int32 TensorIdx = InGraph.OpInputIndices[Idx + InOp.InputStart];				
+			
+			ConnectEdgeDst(TensorIdx, NodeIdx, Idx);
+		}
+
+		for (int32 Idx = 0; Idx < InOp.OutputCount; ++Idx)
+		{
+			const int32 TensorIdx = InGraph.OpOutputIndices[Idx + InOp.OutputStart];
+
+			AddIntermediate(TensorIdx, NodeIdx, Idx);
+		}
+	}
+
+	TArray<FEdge>							Edges;
+	TArray<DML_OPERATOR_GRAPH_NODE_DESC>	Operators;
+	int32									NumInputs;
+	int32									NumOutputs;
+};
+
+//
+//
+//
+FInferenceModelDml::FInferenceModelDml()
+{
+	bUseManualTransitions = true;
+}
+
+//
+//
+//
+FInferenceModelDml::~FInferenceModelDml()
+{
+}
+
+//
+//
+//
+bool FInferenceModelDml::Init(TConstArrayView<uint8> ModelData, FDeviceContextDml* InDevCtx)
+{
+	check(ModelData.Num() > 0);
+	FMLRuntimeFormat	Format;
+
+	if (!LoadModel(ModelData, Format))
+	{
+		return false;
+	}
+
+	DevCtx = InDevCtx;
+	DynamicRHI = GetID3D12PlatformDynamicRHI();
+
+	HRESULT Res = DevCtx->Device->CreateOperatorInitializer(0, nullptr, DML_PPV_ARGS(&OpInit));
+	if (!OpInit)
+	{
+		UE_LOG(LogNNX, Warning, TEXT("Error:Failed to create DML operator initializer"));
+		return false;
+	}
+
+	// DirectML requires all tensors to be concrete
+	// TODO jira 168972: Handle dynamic tensor desc, op should init from symbolic shapes
+	TArray<NNECore::Internal::FTensor>	Tensors;
+
+	Tensors.Reset(AllSymbolicTensorDescs.Num());
+	for (const NNECore::FTensorDesc& TensorDesc : AllSymbolicTensorDescs)
+	{
+		Tensors.Emplace(NNECore::Internal::FTensor::MakeFromSymbolicDesc(TensorDesc));
+	}
+
+	FGraphBuilder				DmlGraphBuilder;
+	FGraphBuilder::FGraphDesc	DmlGraphDesc;
+
+	DmlGraphDesc.AllTensors = Tensors;
+	DmlGraphDesc.InputIndices = InputTensorIndices;
+	DmlGraphDesc.OutputIndices = OutputTensorIndices;
+	DmlGraphDesc.WeightIndices = WeightTensorIndices;
+	DmlGraphDesc.IntermediateIndices = IntermediateTensorIndices;
+	DmlGraphDesc.WeightTensors = WeightTensorRDGs;
+	
+	TArray<FGraphBuilder::FOpDesc>	DmlGraphOperators;
+	TArray<int32>					OpInputIndices;
+	TArray<int32>					OpOutputIndices;
+	uint64							TensorDataSize = 0;
+
+	// Loop over all operators in the model and create them
+	for (int32 Idx = 0; Idx < Format.Operators.Num(); ++Idx)
+	{
+		const FString TypeName = Format.Operators[Idx].TypeName;
+
+		FGraphBuilder::FOpDesc		OpDesc;
+		TArray<NNECore::Internal::FTensor>				OpInputTensors;
+		TArray<NNECore::Internal::FTensor>				OpOutputTensors;
+		NNECore::FAttributeMap	AttributeMap;
+
+		OpDesc.InputStart = OpInputIndices.Num();
+		OpDesc.OutputStart = OpOutputIndices.Num();
+
+		for (int32 InputTensorIndex : Format.Operators[Idx].InTensors)
+		{
+			const int32 WeightTensorIdx = WeightTensorIndices.Find(InputTensorIndex);
+
+			if (WeightTensorIdx >= 0)
+			{
+				TConstArrayView<uint8> TensorData = OpInputTensors.Emplace_GetRef(WeightTensorRDGs[WeightTensorIdx]).GetPreparedData<uint8>();
+				
+				TensorDataSize += Align(TensorData.Num(), DML_MINIMUM_BUFFER_TENSOR_ALIGNMENT);
+			}
+			else
+			{
+				FTensorDesc SymbolicTensorDesc = AllSymbolicTensorDescs[InputTensorIndex];
+				//TODO jira 168972: Handle dynamic tensor desc, op should init from symbolic shapes
+				OpInputTensors.Emplace(NNECore::Internal::FTensor::MakeFromSymbolicDesc(SymbolicTensorDesc));
+			}
+
+			OpInputIndices.Emplace(InputTensorIndex);
+		}
+
+		for (int32 OutputTensorIndex : Format.Operators[Idx].OutTensors)
+		{
+			FTensorDesc SymbolicTensorDesc = AllSymbolicTensorDescs[OutputTensorIndex];
+			//TODO jira 168972: Handle dynamic tensor desc, op should init from symbolic shapes
+			OpOutputTensors.Emplace(NNECore::Internal::FTensor::MakeFromSymbolicDesc(SymbolicTensorDesc));
+			OpOutputIndices.Emplace(OutputTensorIndex);
+		}
+
+		for (const FMLFormatAttributeDesc& Desc : Format.Operators[Idx].Attributes)
+		{
+			AttributeMap.SetAttribute(Desc.Name, Desc.Value);
+		}
+
+		OpDesc.Op = OpCreate(TypeName, OpInputTensors, OpOutputTensors, AttributeMap);
+
+		if (!OpDesc.Op)
+		{
+			UE_LOG(LogNNX, Warning, TEXT("Error:Failed to create operator:%s"), *TypeName);
+			return false;
+		}
+
+		OpDesc.InputCount = OpInputTensors.Num();
+		OpDesc.OutputCount = OpOutputTensors.Num();
+		OpDesc.DbgName = TypeName;
+
+		DmlGraphOperators.Emplace(OpDesc);
+	}
+
+	DmlGraphDesc.Operators = DmlGraphOperators;
+	DmlGraphDesc.OpInputIndices = OpInputIndices;
+	DmlGraphDesc.OpOutputIndices = OpOutputIndices;
+
+	CompiledOp = DmlGraphBuilder.Compile(DevCtx, DmlGraphDesc);
+	if (!CompiledOp)
+	{
+		return false;
+	}
+
+	return InitCompiledOp(OpInputIndices, TensorDataSize);
+}
+
+//
+//
+//
+bool FInferenceModelDml::InitCompiledOp(TConstArrayView<int32> OpInputIndices, uint64 TensorDataSize)
+{
+	static constexpr EBufferUsageFlags	WeightBuffUsage = BUF_UnorderedAccess;
+	static constexpr ERHIAccess			WeightBuffAccess = ERHIAccess::UAVMask;
+
+	static constexpr EBufferUsageFlags	PersistBuffFlags = BUF_Static | BUF_ShaderResource | BUF_UnorderedAccess;
+	static constexpr ERHIAccess			PersistBuffAccess = ERHIAccess::UAVMask;
+
+	static constexpr EBufferUsageFlags	TempBuffFlags = BUF_Volatile | BUF_UnorderedAccess;
+	static constexpr ERHIAccess			TempBuffAccess = ERHIAccess::UAVMask;
+
+	HRESULT					Res;
+	IDMLDevice*				Device = DevCtx->Device;
+	IDMLCompiledOperator*	CompiledOps[] = { CompiledOp };
+	
+	Res = OpInit->Reset(UE_ARRAY_COUNT(CompiledOps), CompiledOps);
+	if (FAILED(Res))
+	{
+		UE_LOG(LogNNX, Warning, TEXT("Error:Failed to reset DirectML operator initializer"));
+		return false;
+	}
+
+	DML_BINDING_PROPERTIES InitBindProps = OpInit->GetBindingProperties();
+	DML_BINDING_PROPERTIES ExecBindProps = CompiledOp->GetBindingProperties();
+
+	DescCount = std::max(InitBindProps.RequiredDescriptorCount, ExecBindProps.RequiredDescriptorCount);
+	
+	D3D12_DESCRIPTOR_HEAP_DESC	HeapDesc = {};
+
+	HeapDesc.Flags = D3D12_DESCRIPTOR_HEAP_FLAG_SHADER_VISIBLE;
+	HeapDesc.Type = D3D12_DESCRIPTOR_HEAP_TYPE_CBV_SRV_UAV;
+	HeapDesc.NumDescriptors = DescCount;
+
+	Res = DevCtx->D3D12Device->CreateDescriptorHeap(&HeapDesc, IID_PPV_ARGS(&DescHeap));
+	if (!DescHeap)
+	{
+		UE_LOG(LogNNX, Warning, TEXT("Failed to create descriptor heap, res:%x"), Res);
+		return false;
+	}
+
+	DescSize = DevCtx->D3D12Device->GetDescriptorHandleIncrementSize(D3D12_DESCRIPTOR_HEAP_TYPE_CBV_SRV_UAV);
+
+	BindingTable.Reset(new FBindingTable());
+	if (!BindingTable->Init(this))
+	{
+		return false;
+	}
+
+	MemSizeTemp = ExecBindProps.TemporaryResourceSize;
+	MemSizePersist = ExecBindProps.PersistentResourceSize;
+
+	FEvent* Signal = FGenericPlatformProcess::GetSynchEventFromPool(false);
+
+	ENQUEUE_RENDER_COMMAND(FInferenceModelDml_SetTensorData)
+	(
+		[
+			this, 
+			Signal, 
+			InitTempMemSize = InitBindProps.TemporaryResourceSize,
+			TensorDataSize
+		]
+		(FRHICommandListImmediate& RHICmdList)
+		{
+			FRHIBufferInputArray	Inputs;
+
+			for (int32 InputIdx : InputTensorIndices)
+			{
+				Inputs.Emplace(nullptr);
+			}
+
+			TArray<CD3DX12_RESOURCE_BARRIER, TInlineAllocator<MaxNumInputs>>	Barriers;
+			FGPUFenceRHIRef	UploadFence = nullptr;
+
+			if (TensorDataSize)
+			{
+				UploadFence = RHICreateGPUFence(TEXT("FInferenceModel_UploadFence"));
+				
+				FRHIBuffer*		UploadBuff = CreateRHIBuffer(RHICmdList, TensorDataSize, BUF_ShaderResource | BUF_Dynamic | BUF_FastVRAM, ERHIAccess::CopySrc, TEXT("FInferenceModel_UploadBuffer"));
+				uint8*			UploadBuffPtr = static_cast<uint8*>(RHICmdList.LockBuffer(UploadBuff, 0, TensorDataSize, RLM_WriteOnly_NoOverwrite));
+				uint64			UploadOffset = 0;
+
+				for (const FTensorRDG& Tensor : WeightTensorRDGs)
+				{
+					TConstArrayView<uint8>	TensorData = Tensor.GetPreparedData<uint8>();
+
+					FBufferRHIRef WeightBuff;
+
+					WeightBuff = CreateRHIBuffer(RHICmdList, TensorData.Num(), WeightBuffUsage, WeightBuffAccess, TEXT("FInferenceModelDml_TensorWeights"));
+				
+					FMemory::Memcpy(UploadBuffPtr + UploadOffset, TensorData.GetData(), TensorData.Num());
+					RHICmdList.CopyBufferRegion(WeightBuff, 0, UploadBuff, UploadOffset, TensorData.Num());
+					UploadOffset += Align(TensorData.Num(), DML_MINIMUM_BUFFER_TENSOR_ALIGNMENT);
+
+					Inputs.Emplace(WeightBuff);
+				
+					Barriers.Emplace(
+						CD3DX12_RESOURCE_BARRIER::Transition(
+							DynamicRHI->RHIGetResource(WeightBuff),
+							D3D12_RESOURCE_STATE_COPY_DEST,
+							D3D12_RESOURCE_STATE_UNORDERED_ACCESS)
+					);
+				}
+
+				RHICmdList.UnlockBuffer(UploadBuff);
+				RHICmdList.WriteGPUFence(UploadFence);
+			}
+
+			if (MemSizePersist)
+			{
+#ifdef NNE_USE_D3D12_RESOURCES
+				PersistBuff = CreateD3D12Buffer(MemSizePersist);
+#else		
+				PersistBuff = CreateRHIBuffer(RHICmdList, MemSizePersist, PersistBuffFlags, PersistBuffAccess, TEXT("FMLInferendeModelDml_PeristBuff"));
+#endif
+			}
+
+			if (MemSizeTemp)
+			{
+#ifdef NNE_USE_D3D12_RESOURCES
+				TempBuff = CreateD3D12Buffer(MemSizeTemp);
+#else
+				TempBuff = CreateRHIBuffer(RHICmdList, MemSizeTemp, TempBuffFlags, TempBuffAccess, TEXT("FMLInferendeModelDml_TempBuff"));
+#endif
+			}
+
+#ifdef NNE_USE_D3D12_RESOURCES
+			ID3D12Resource* InitTempBuff = nullptr;
+#else
+			FBufferRHIRef InitTempBuff;
+#endif
+
+			if (InitTempMemSize)
+			{
+#ifdef NNE_USE_D3D12_RESOURCES
+				InitTempBuff = CreateD3D12Buffer(InitTempMemSize);
+#else
+				TempBuff = CreateRHIBuffer(RHICmdList, InitTempMemSize, TempBuffFlags, TempBuffAccess, TEXT("FMLInferendeModelDml_InitTempBuff"));
+#endif
+			}
+
+			RHICmdList.EnqueueLambda(
+				[this, Inputs, Barriers, InitTempBuff, UploadFence](FRHICommandListImmediate& RHICmdList)
+				{
+					while (UploadFence && UploadFence->NumPendingWriteCommands.GetValue() > 0)
+					{
+						FPlatformProcess::Sleep(0.001);
+					}
+
+					ID3D12GraphicsCommandList* D3DCmdList = nullptr;
+
+					D3DCmdList = DynamicRHI->RHIGetGraphicsCommandList(DevCtx->DeviceIndex);
+
+					BindingTable->Bind(OpInit, Inputs, PersistBuff, InitTempBuff);
+			
+					D3DCmdList = DynamicRHI->RHIGetGraphicsCommandList(DevCtx->DeviceIndex);
+
+					if (!Barriers.IsEmpty())
+					{
+						D3DCmdList->ResourceBarrier(Barriers.Num(), Barriers.GetData());
+					}
+					D3DCmdList->SetDescriptorHeaps(1, &DescHeap);
+					DevCtx->CmdRec->RecordDispatch(D3DCmdList, OpInit, BindingTable->Get());
+
+					DynamicRHI->RHIFinishExternalComputeWork(DevCtx->DeviceIndex, D3DCmdList);
+				}
+			);			
+			
+			RHICmdList.ImmediateFlush(EImmediateFlushType::FlushRHIThread);
+			Signal->Trigger();
+		}
+	);
+
+	Signal->Wait();
+	FGenericPlatformProcess::ReturnSynchEventToPool(Signal);
+
+	return true;
+}
+
+BEGIN_SHADER_PARAMETER_STRUCT(FTensorBufferParamsDml, )
+	RDG_BUFFER_ACCESS(Buffer, ERHIAccess::UAVCompute)
+END_SHADER_PARAMETER_STRUCT()
+
+//
+//
+//
+void FInferenceModelDml::AddDispatchOps_RenderThread(FRDGBuilder& GraphBuilder)
+{
+	const ERDGPassFlags TransitionBuffFlags = ERDGPassFlags::Compute | ERDGPassFlags::NeverCull;
+
+	InputBuffers.SetNumUninitialized(InputTensorIndices.Num() + WeightTensorIndices.Num());
+	OutputBuffers.SetNumUninitialized(OutputTensorIndices.Num());
+
+	for (int32 Idx = 0; Idx < InputTensorIndices.Num(); ++Idx)
+	{
+		FTensorBufferParamsDml* Params = GraphBuilder.AllocParameters<FTensorBufferParamsDml>();
+		Params->Buffer = AllTensorRDGs[InputTensorIndices[Idx]]->GetBuffer();
+
+		GraphBuilder.AddPass(
+			RDG_EVENT_NAME("FInferenceModelDml_Dispatch_GetInputBuffer"),
+			Params,
+			TransitionBuffFlags,
+			[this, Idx, Params](FRHICommandListImmediate& RHICmdList)
+			{
+				InputBuffers[Idx] = Params->Buffer->GetRHI();
+			}
+		);
+	}
+
+	for (int32 Idx = 0; Idx < WeightTensorIndices.Num(); ++Idx)
+	{
+		InputBuffers[Idx + InputTensorIndices.Num()] = nullptr;
+	}
+
+	for (int32 Idx = 0; Idx < OutputTensorIndices.Num(); ++Idx)
+	{
+		FTensorBufferParamsDml* Params = GraphBuilder.AllocParameters<FTensorBufferParamsDml>();
+		Params->Buffer = AllTensorRDGs[OutputTensorIndices[Idx]]->GetBuffer();
+
+		GraphBuilder.AddPass(
+			RDG_EVENT_NAME("FInferenceModelDml_Dispatch_GetOutputBuffer"),
+			Params,
+			TransitionBuffFlags,
+			[this, Idx, Params](FRHICommandListImmediate& RHICmdList)
+			{
+				OutputBuffers[Idx] = Params->Buffer->GetRHI();
+			}
+		);
+	}
+	
+	GraphBuilder.AddPass(
+		RDG_EVENT_NAME("FInferenceModelDml_Dispatch"),
+		ERDGPassFlags::None | ERDGPassFlags::NeverCull,
+		[this](FRHICommandListImmediate& RHICmdList)
+		{
+			RHICmdList.EnqueueLambda(
+				[this](FRHICommandListImmediate& RHICmdList)
+				{
+					TArray<CD3DX12_RESOURCE_BARRIER, TInlineAllocator<MaxNumInputs + MaxNumOutputs>>	Barriers;
+
+					for (FRHIBuffer* Buffer : InputBuffers)
+					{
+						if (!Buffer)
+						{
+							continue;
+						}
+
+						ID3D12Resource* Resource = DynamicRHI->RHIGetResource(Buffer);
+
+						// TODO: FIXME: Don't assume COPY_DEST state
+						Barriers.Emplace(
+							CD3DX12_RESOURCE_BARRIER::Transition(
+								Resource,
+								D3D12_RESOURCE_STATE_COPY_DEST,
+								D3D12_RESOURCE_STATE_UNORDERED_ACCESS)
+						);
+					}
+
+					// TODO: We should use this instead of NNE_USE_D3D12_RESOURCES
+					//FBufferRHIRef TempBuff = MemSizeTemp ? CreateRHIBuffer(RHICmdList, MemSizeTemp, TempBuffUsage, TempBuffAccess, TEXT("FInferenceModelDml_Dispatch_TempBuff")) : nullptr;
+
+					BindingTable->Bind(CompiledOp, InputBuffers, OutputBuffers, PersistBuff, TempBuff);
+
+					ID3D12GraphicsCommandList* D3DCmdList = nullptr;
+
+					D3DCmdList = DynamicRHI->RHIGetGraphicsCommandList(DevCtx->DeviceIndex);
+					D3DCmdList->SetDescriptorHeaps(1, &DescHeap);
+					D3DCmdList->ResourceBarrier(Barriers.Num(), Barriers.GetData());
+					DevCtx->CmdRec->RecordDispatch(D3DCmdList, CompiledOp, BindingTable->Get());
+
+					DynamicRHI->RHIFinishExternalComputeWork(DevCtx->DeviceIndex, D3DCmdList);
+				}
+			);
+		}
+	);
+}
+
+//
+// Create operator
+//
+FOperatorDml* FInferenceModelDml::OpCreate(const FString& OpName, TArrayView<const NNECore::Internal::FTensor> InputTensorDescs, TArrayView<const NNECore::Internal::FTensor> OutputTensorDescs, const NNECore::FAttributeMap& Attributes)
+{
+	FOperatorRegistryDml::OperatorCreateFunc CreateFn = FOperatorRegistryDml::Get()->OpFind(OpName);
+
+	if (!CreateFn)
+	{
+		UE_LOG(LogNNX, Warning, TEXT("Dml MLOperatorRegistry failed to find operator:%s"), *OpName);
+		return nullptr;
+	}
+
+	FOperatorDml* Op = CreateFn();
+
+	if (!Op->Initialize(DevCtx, InputTensorDescs, OutputTensorDescs, Attributes))
+	{
+		delete Op;
+
+		UE_LOG(LogNNX, Warning, TEXT("Error:Failed to initialize operator:%s"), *OpName);
+		return nullptr;
+	}
+
+	Op->GetOperator()->SetName(*OpName);
+
+	return Op;
+}
+
+FBufferRHIRef FInferenceModelDml::CreateRHIBuffer(FRHICommandListImmediate& RHICmdList, uint32 Size, EBufferUsageFlags Usage, ERHIAccess Access, const TCHAR* DbgName)
+{
+	FBufferRHIRef Buff = nullptr;
+
+	if (Size)
+	{
+		FRHIResourceCreateInfo CreateInfo(DbgName);
+		
+		Buff = RHICmdList.CreateBuffer(Size, Usage, 1, Access, CreateInfo);
+	}
+
+	check(Buff);
+	return Buff;
+}
+
+ID3D12Resource* FInferenceModelDml::CreateD3D12Buffer(uint32 Size, D3D12_RESOURCE_STATES ResourceState, D3D12_HEAP_TYPE HeapType, const TCHAR* DebugName)
+{
+	ID3D12Resource* Resource = nullptr;
+
+	D3D12_RESOURCE_DESC ResourceDesc = CD3DX12_RESOURCE_DESC::Buffer(Size,D3D12_RESOURCE_FLAG_ALLOW_UNORDERED_ACCESS);
+	CD3DX12_HEAP_PROPERTIES HeapProps = CD3DX12_HEAP_PROPERTIES(HeapType);
+	HRESULT Res;
+
+	Res = DevCtx->D3D12Device->CreateCommittedResource(
+		&HeapProps,
+		D3D12_HEAP_FLAG_NONE,
+		&ResourceDesc,
+		ResourceState,
+		nullptr,
+		IID_PPV_ARGS(&Resource));
+
+	if (FAILED(Res))
+	{
+		UE_LOG(LogNNX, Warning, TEXT("Error:FInferenceModel failed to create D3D12 resource"));
+		return nullptr;
+	}
+
+	if (Resource && DebugName)
+	{
+		Resource->SetName(DebugName);
+	}
+
+	return Resource;
+}
+
+//
+//
+//
+int FInferenceModelDml::PrepareTensorShapesAndData()
+{
+	for (FTensorDesc SymbolicTensorDesc : AllSymbolicTensorDescs)
+	{
+		if (!SymbolicTensorDesc.GetShape().IsConcrete())
+		{
+			UE_LOG(LogNNX, Warning, TEXT("DML engine does not support model with variable shapes yet."));
+			return -1;
+		}
+	}
+
+	return 0;
+}
+
+//
+//
+//
+static TUniquePtr<FRuntimeDml> FDmlRuntimeCreate(bool bRegisterOnlyOperators)
+{
+	auto Runtime = MakeUnique<FRuntimeDml>();
+
+	if (!Runtime->Init(bRegisterOnlyOperators))
+	{
+		UE_LOG(LogNNX, Warning, TEXT("Failed to create NNX DML runtime"));
+		Runtime.Release();
+	}
+
+	return Runtime;
+}
+
+//
+// Called on RDG runtime startup
+//
+IRuntime* FRuntimeDmlStartup()
+{
+	if (!GDmlRuntime)
+	{
+		bool bIsD3D12RHI = GDynamicRHI && GDynamicRHI->GetInterfaceType() == ERHIInterfaceType::D3D12;		
+		bool bLoadDirectML = true;
+
+		if (IsRunningCommandlet() && !IsAllowCommandletRendering())
+		{
+			UE_LOG(LogNNX, Display, TEXT("Running inside commandlet without rendering"));
+			bLoadDirectML = false;
+		}
+
+#ifdef DIRECTML_BIN_PATH
+		
+		if (bIsD3D12RHI && bLoadDirectML)
+		{
+			const FString DirectMLRuntimeBinPath = TEXT(PREPROCESSOR_TO_STRING(DIRECTML_BIN_PATH));
+			FString DirectMLDLLPaths[2];
+			int32	NumPaths = 1;
+		
+			DirectMLDLLPaths[0] = DirectMLRuntimeBinPath / TEXT("DirectML.dll");
+
+			if (GetID3D12PlatformDynamicRHI()->IsD3DDebugEnabled())
+			{
+				DirectMLDLLPaths[1] = DirectMLRuntimeBinPath / TEXT("DirectML.Debug.dll");
+				++NumPaths;
+			}
+
+			FPlatformProcess::PushDllDirectory(*DirectMLRuntimeBinPath);
+
+			for (int32 Idx = 0; Idx < NumPaths; ++Idx)
+			{
+				if (!FPaths::FileExists(DirectMLDLLPaths[Idx]))
+				{
+					const FString ErrorMessage = FString::Format(TEXT("DirectML DLL file not found in \"{0}\"."),
+						{ IFileManager::Get().ConvertToAbsolutePathForExternalAppForRead(*DirectMLDLLPaths[Idx])});
+					UE_LOG(LogNNX, Warning, TEXT("NNXRuntimeDll:%s"), *ErrorMessage);
+					checkf(false, TEXT("%s"), *ErrorMessage);
+				}
+
+				FPlatformProcess::GetDllHandle(*DirectMLDLLPaths[Idx]);
+			}
+
+			FPlatformProcess::PopDllDirectory(*DirectMLRuntimeBinPath);
+		}
+#endif
+
+		const bool bRegisterOnlyOperators = !bLoadDirectML;
+
+		GDmlRuntime = FDmlRuntimeCreate(bRegisterOnlyOperators);
+	}
+
+	return GDmlRuntime.Get();
+}
+
+//
+// Called on RDG runtime shutdown
+//
+void FRuntimeDmlShutdown()
+{
+	if (GDmlRuntime)
+	{
+		GDmlRuntime.Release();
+	}
+}
+
+} // namespace UE::NNERuntimeRDG::Private::Dml
+ 
+#endif
diff --git a/Engine/Plugins/Experimental/NNE/Source/NNERuntimeRDG/Private/NNERuntimeRDGHlsl.cpp b/Engine/Plugins/Experimental/NNE/Source/NNERuntimeRDG/Private/NNERuntimeRDGHlsl.cpp
new file mode 100644
--- /dev/null
+++ b/Engine/Plugins/Experimental/NNE/Source/NNERuntimeRDG/Private/NNERuntimeRDGHlsl.cpp
@@ -0,0 +1,154 @@
+// Copyright Epic Games, Inc. All Rights Reserved.
+
+#include "Hlsl/NNERuntimeRDGConv.h"
+#include "Hlsl/NNERuntimeRDGConv.h"
+#include "Hlsl/NNERuntimeRDGConvTranspose.h"
+#include "Hlsl/NNERuntimeRDGElementWiseBinary.h"
+#include "Hlsl/NNERuntimeRDGElementWiseUnary.h"
+#include "Hlsl/NNERuntimeRDGElementWiseVariadic.h"
+#include "Hlsl/NNERuntimeRDGGemm.h"
+#include "Hlsl/NNERuntimeRDGInstanceNormalization.h"
+#include "Hlsl/NNERuntimeRDGPad.h"
+#include "Hlsl/NNERuntimeRDGUpsample.h"
+#include "Hlsl/NNERuntimeRDGMatMul.h"
+#include "NNECoreAttributeMap.h"
+#include "NNERuntimeRDGModelHlsl.h"
+#include "NNXInferenceModel.h"
+#include "NNEUtilsModelOptimizer.h"
+#include "NNXRuntimeFormat.h"
+#include "NNERuntimeRDGHlslHelper.h"
+#include "NNERuntimeRDG.h"
+
+namespace UE::NNERuntimeRDG::Private::Hlsl
+{
+
+#define NNX_RUNTIME_HLSL_NAME TEXT("NNXRuntimeHlsl")
+
+//
+//
+//
+class FRuntimeHlsl : public FRuntimeRDG
+{
+public:
+
+	FRuntimeHlsl() = default;
+
+	virtual ~FRuntimeHlsl()
+	{
+	}
+
+	bool Init()
+	{
+		FOperatorRegistryHlsl* registry = FOperatorRegistryHlsl::Get();
+		check(registry != nullptr);
+
+		RegisterElementWiseUnaryOperators(*registry);
+		RegisterElementWiseBinaryOperators(*registry);
+		RegisterElementWiseVariadicOperators(*registry);
+		RegisterGemmOperator(*registry);
+		RegisterConvOperator(*registry);
+		RegisterConvTransposeOperator(*registry);
+		RegisterMatMulOperator(*registry);
+		RegisterInstanceNormalizationOperator(*registry);
+		RegisterUpsampleOperator(*registry);
+		RegisterPadOperator(*registry);
+
+		return true;
+	}
+
+	virtual FString GetRuntimeName() const override
+	{
+		return NNX_RUNTIME_HLSL_NAME;
+	}
+
+	// Returns flags from ERuntimeSupportFlags
+	virtual NNX::EMLRuntimeSupportFlags GetSupportFlags() const override
+	{
+		return NNX::EMLRuntimeSupportFlags::RDG;
+	}
+
+	virtual TArray<uint8> CreateModelData(FString FileType, TConstArrayView<uint8> FileData) override
+	{
+		if (!CanCreateModelData(FileType, FileData))
+		{
+			return {};
+		}
+
+		TUniquePtr<NNX::IModelOptimizer> Optimizer = NNEUtils::Internal::CreateONNXToNNEModelOptimizer();
+		Optimizer->AddValidator(MakeShared<FModelValidatorHlsl>());
+
+		FNNIModelRaw InputModel;
+		InputModel.Data = FileData;
+		InputModel.Format = ENNXInferenceFormat::ONNX;
+
+		FNNIModelRaw OutputModel;
+		if (!Optimizer->Optimize(InputModel, OutputModel, {}))
+		{
+			return {};
+		}
+
+		return ConvertToModelData(OutputModel.Data);
+	};
+
+	virtual TUniquePtr<NNX::FMLInferenceModel> CreateModel(TConstArrayView<uint8> ModelData) override
+	{
+		if (!CanCreateModel(ModelData))
+		{
+			return TUniquePtr<NNX::FMLInferenceModel>();
+		}
+
+		// Create the model and initialize it with the data not including the header
+		FModel* Model = new FModel();
+		if (!Model->Init(ModelData))
+		{
+			delete Model;
+			return TUniquePtr<NNX::FMLInferenceModel>();
+		}
+		return TUniquePtr<NNX::FMLInferenceModel>(Model);
+	}
+};
+
+/** Globally accessible runtime */
+static TUniquePtr<NNERuntimeRDG::Private::Hlsl::FRuntimeHlsl> GHlslRuntime;
+
+//
+//
+//
+static TUniquePtr<NNERuntimeRDG::Private::Hlsl::FRuntimeHlsl> FRuntimeHlslCreate()
+{
+	auto Runtime = MakeUnique<NNERuntimeRDG::Private::Hlsl::FRuntimeHlsl>();
+
+	if (!Runtime->Init())
+	{
+		UE_LOG(LogNNX, Warning, TEXT("Failed to create NNX HLSL runtime"));
+		Runtime.Release();
+	}
+
+	return Runtime;
+}
+
+//
+// Called on RDG runtime startup
+//
+IRuntime* FRuntimeHlslStartup()
+{
+	if (!GHlslRuntime)
+	{
+		GHlslRuntime = FRuntimeHlslCreate();
+	}
+
+	return GHlslRuntime.Get();
+}
+
+//
+// Called on RDG runtime shutdown
+//
+void FRuntimeHlslShutdown()
+{
+	if (GHlslRuntime)
+	{
+		GHlslRuntime.Release();
+	}
+}
+
+} // namespace UE::NNERuntimeRDG::Private::Hlsl
diff --git a/Engine/Plugins/Experimental/NNE/Source/NNERuntimeRDG/Private/NNERuntimeRDGHlslHelper.cpp b/Engine/Plugins/Experimental/NNE/Source/NNERuntimeRDG/Private/NNERuntimeRDGHlslHelper.cpp
new file mode 100644
--- /dev/null
+++ b/Engine/Plugins/Experimental/NNE/Source/NNERuntimeRDG/Private/NNERuntimeRDGHlslHelper.cpp
@@ -0,0 +1,75 @@
+// Copyright Epic Games, Inc. All Rights Reserved.
+
+#include "NNERuntimeRDGHlslHelper.h"
+#include "NNECoreTypes.h"
+#include "NNECoreTensor.h"
+#include "RHI.h"
+
+namespace UE::NNERuntimeRDG::Private::Hlsl
+{
+	void FillTensorSizeShaderParameters(const NNECore::Internal::FTensor& Tensor, TStaticArray<FUintVector4, NXRT_TENSORSTRIDEINFO_MAX_NUM_DIMENSIONS, 16U>& OutShaderParam, int32 Idx)
+	{
+		static_assert(NNECore::FTensorShape::MaxRank <= NXRT_TENSORSTRIDEINFO_MAX_NUM_DIMENSIONS);
+		for (int32 i = 0; i < Tensor.GetShape().Rank(); ++i)
+		{
+			OutShaderParam[i][Idx] = Tensor.GetShape().GetData()[i];
+		}
+	}
+
+	void FillTensorStrideShaderParameters(const NNECore::Internal::FTensor& Tensor, TStaticArray<FUintVector4, NXRT_TENSORSTRIDEINFO_MAX_NUM_DIMENSIONS, 16U>& OutShaderParam, int32 Idx, int32 TargetNumdimensionForBroadcast)
+	{
+		if (TargetNumdimensionForBroadcast == -1)
+		{
+			TargetNumdimensionForBroadcast = Tensor.GetShape().Rank();
+		}
+		checkf(TargetNumdimensionForBroadcast >= Tensor.GetShape().Rank(), TEXT("Can't broadcast tensor from rank %d to rank %d, should be inferior or equal."), Tensor.GetShape().Rank(), TargetNumdimensionForBroadcast);
+		int32 Offset = TargetNumdimensionForBroadcast - Tensor.GetShape().Rank();
+
+		static_assert(NNECore::FTensorShape::MaxRank <= NXRT_TENSORSTRIDEINFO_MAX_NUM_DIMENSIONS);
+		for (int32 i = 7; i >= 0; --i)
+		{
+			if (i >= TargetNumdimensionForBroadcast || i < Offset)
+			{
+				OutShaderParam[i][Idx] = 0;
+			}
+			else if (i == TargetNumdimensionForBroadcast - 1)
+			{
+				OutShaderParam[i][Idx] = 1;
+			}
+			else
+			{
+				OutShaderParam[i][Idx] = OutShaderParam[i + 1][Idx] * Tensor.GetShape().GetData()[i + 1 - Offset];
+			}
+		}
+	}
+
+	void FillTensorStrideForBroadcastShaderParameters(const NNECore::Internal::FTensor& Tensor, int32 OutputNumdimension, TStaticArray<FUintVector4, NXRT_TENSORSTRIDEINFO_MAX_NUM_DIMENSIONS, 16U>& OutShaderParam, int32 Idx)
+	{
+		checkf(OutputNumdimension >= Tensor.GetShape().Rank(), TEXT("Can't broadcast tensor from rank %d to rank %d, should be inferior or equal."), Tensor.GetShape().Rank(), OutputNumdimension);
+		FillTensorStrideShaderParameters(Tensor, OutShaderParam, Idx, OutputNumdimension);
+		int32 Offset = OutputNumdimension - Tensor.GetShape().Rank();
+		for (int32 i = Offset; i < OutputNumdimension; ++i)
+		{
+			// the stride for broadcast dimension is kept as 0
+			if (Tensor.GetShape().GetData()[i - Offset] == 1)
+			{
+				OutShaderParam[i][Idx] = 0;
+			}
+		}
+	}
+
+	FIntVector ComputeElementWiseThreadGroups(uint32 ElementCount, uint32 GroupSizeX)
+	{
+		FIntVector ThreadGroupCount;
+		ThreadGroupCount.X = FMath::DivideAndRoundUp(ElementCount, GroupSizeX);
+		ThreadGroupCount.Y = 1;
+		ThreadGroupCount.Z = 1;
+		if (ThreadGroupCount.X > GRHIMaxDispatchThreadGroupsPerDimension.X)
+		{
+			ThreadGroupCount.Y = FMath::DivideAndRoundUp(ThreadGroupCount.X, GRHIMaxDispatchThreadGroupsPerDimension.X);
+			ThreadGroupCount.X = FMath::DivideAndRoundUp(ThreadGroupCount.X, ThreadGroupCount.Y);
+			ensure(ThreadGroupCount.Y <= GRHIMaxDispatchThreadGroupsPerDimension.Y);
+		}
+		return ThreadGroupCount;
+	}
+} // UE::NNERuntimeRDG::Private::Hlsl
\ No newline at end of file
diff --git a/Engine/Plugins/Experimental/NNE/Source/NNERuntimeRDG/Private/NNERuntimeRDGHlslHelper.h b/Engine/Plugins/Experimental/NNE/Source/NNERuntimeRDG/Private/NNERuntimeRDGHlslHelper.h
new file mode 100644
--- /dev/null
+++ b/Engine/Plugins/Experimental/NNE/Source/NNERuntimeRDG/Private/NNERuntimeRDGHlslHelper.h
@@ -0,0 +1,18 @@
+// Copyright Epic Games, Inc. All Rights Reserved.
+
+#pragma once
+
+#include "Containers/StaticArray.h"
+
+namespace UE::NNECore::Internal { class FTensor; }
+
+#define NXRT_TENSORSTRIDEINFO_MAX_NUM_DIMENSIONS 8
+
+namespace UE::NNERuntimeRDG::Private::Hlsl
+{
+
+	void FillTensorSizeShaderParameters(const NNECore::Internal::FTensor& Tensor, TStaticArray<FUintVector4, NXRT_TENSORSTRIDEINFO_MAX_NUM_DIMENSIONS, 16U>& OutShaderParam, int32 Idx);
+	void FillTensorStrideShaderParameters(const NNECore::Internal::FTensor& Tensor, TStaticArray<FUintVector4, NXRT_TENSORSTRIDEINFO_MAX_NUM_DIMENSIONS, 16U>& OutShaderParam, int32 Idx, int32 TargetNumdimensionForBroadcast = -1);
+	void FillTensorStrideForBroadcastShaderParameters(const NNECore::Internal::FTensor& Tensor, int32 OutputNumdimension, TStaticArray<FUintVector4, NXRT_TENSORSTRIDEINFO_MAX_NUM_DIMENSIONS, 16U>& OutShaderParam, int32 Idx);
+	FIntVector ComputeElementWiseThreadGroups(uint32 ElementCount, uint32 GroupSizeX);
+} // NNX
\ No newline at end of file
diff --git a/Engine/Plugins/Experimental/NNE/Source/NNERuntimeRDG/Private/NNERuntimeRDGHlslOp.h b/Engine/Plugins/Experimental/NNE/Source/NNERuntimeRDG/Private/NNERuntimeRDGHlslOp.h
new file mode 100644
--- /dev/null
+++ b/Engine/Plugins/Experimental/NNE/Source/NNERuntimeRDG/Private/NNERuntimeRDGHlslOp.h
@@ -0,0 +1,26 @@
+// Copyright Epic Games, Inc. All Rights Reserved.
+
+#pragma once
+
+#include "NNXRuntime.h"
+#include "NNERuntimeRDG.h"
+
+namespace UE::NNERuntimeRDG::Private::Hlsl
+{
+	/**
+	 * Base class for all Hlsl ML operators
+	 */
+	struct FOperatorHlsl : public IOperatorRDG, public IPrepareOperator
+	{
+		virtual ~FOperatorHlsl() = default;
+		virtual bool Initialize(TConstArrayView<FTensorDesc> InputTensorDescs, TConstArrayView<FTensorDesc> OutputTensorDescs, const NNECore::FAttributeMap& Attributes) = 0;
+		virtual void Dispatch(FRDGBuilder& GraphBuilder, TConstArrayView<FTensorRDGRef> InInputTensors, TConstArrayView<FTensorRDGRef> InOutputTensors) = 0;
+	};
+
+	/**
+	  * HLSL ML operator registry
+	*/
+	typedef TOperatorRegistryRDG<FOperatorHlsl> FOperatorRegistryHlsl;
+	typedef TModelValidatorRDG<FOperatorHlsl> FModelValidatorHlsl;
+
+} // UE::NNERuntimeRDG::Private::Hlsl
\ No newline at end of file
diff --git a/Engine/Plugins/Experimental/NNE/Source/NNERuntimeRDG/Private/NNERuntimeRDGModelHlsl.cpp b/Engine/Plugins/Experimental/NNE/Source/NNERuntimeRDG/Private/NNERuntimeRDGModelHlsl.cpp
new file mode 100644
--- /dev/null
+++ b/Engine/Plugins/Experimental/NNE/Source/NNERuntimeRDG/Private/NNERuntimeRDGModelHlsl.cpp
@@ -0,0 +1,444 @@
+// Copyright Epic Games, Inc. All Rights Reserved.
+
+#include "NNERuntimeRDGModelHlsl.h"
+#include "NNECoreTensor.h"
+#include "NNERuntimeRDGHlslOp.h"
+#include "RenderGraphBuilder.h"
+#include "RenderGraphUtils.h"
+
+namespace UE::NNERuntimeRDG::Private::Hlsl
+{
+
+namespace
+{
+
+FOperatorHlsl* OpCreate(const FString& OpName, TConstArrayView<NNX::FTensorDesc> InputTensorDescs, TConstArrayView<NNX::FTensorDesc> OutputTensorDescs, const NNECore::FAttributeMap& AttributeMap)
+{
+	FOperatorRegistryHlsl::OperatorCreateFunc CreateFn = FOperatorRegistryHlsl::Get()->OpFind(OpName);
+
+	if (!CreateFn)
+	{
+		UE_LOG(LogNNX, Warning, TEXT("Hlsl MLOperatorRegistry failed to find operator:%s"), *OpName);
+		return nullptr;
+	}
+
+	FOperatorHlsl* Op = CreateFn();
+
+	if (!Op->Initialize(InputTensorDescs, OutputTensorDescs, AttributeMap))
+	{
+		UE_LOG(LogNNX, Warning, TEXT("Hlsl engine: Error initializing operator:%s"), *OpName);
+		delete Op;
+		return nullptr;
+	}
+
+	return Op;
+}
+
+template<class OperatorRef>
+void InternAddDispatchOps_RenderThread(FRDGBuilder& GraphBuilder, TConstArrayView<FTensorHLSLRef> AllTensorHLSLRefs, TConstArrayView<TArray<uint32>> OperatorInputTensorIndices,
+	TConstArrayView<TArray<uint32>> OperatorOutputTensorIndices, TArrayView<OperatorRef> Operators)
+{
+	static constexpr int32 MaxExpectedInput = 10;
+	TArray<FTensorRDGRef, TInlineAllocator<MaxExpectedInput>> InputTensors;
+
+	static constexpr int32 MaxExpectedOutput = 2;
+	TArray<FTensorRDGRef, TInlineAllocator<MaxExpectedOutput>> OutputTensors;
+
+	// Add passes for all operators
+	for (int32 Idx = 0; Idx < Operators.Num(); ++Idx)
+	{
+		InputTensors.Empty();
+		for (int32 i : OperatorInputTensorIndices[Idx])
+		{
+			InputTensors.Add(AllTensorHLSLRefs[i]);
+		}
+		OutputTensors.Empty();
+		for (int32 i : OperatorOutputTensorIndices[Idx])
+		{
+			OutputTensors.Add(AllTensorHLSLRefs[i]);
+		}
+
+		Operators[Idx]->Dispatch(GraphBuilder, InputTensors, OutputTensors);
+	}
+}
+
+int ApplyBinding(TArray<FTensorHLSL>& OutTensorsHLSL, TConstArrayView<NNX::FMLTensorBinding> InBindings, bool bIsInput)
+{
+	check(OutTensorsHLSL.Num() == InBindings.Num());
+	
+	for (int32 Idx = 0; Idx < InBindings.Num(); ++Idx)
+	{
+		const NNX::FMLTensorBinding& Binding = InBindings[Idx];
+		FTensorHLSL& Tensor = OutTensorsHLSL[Idx];
+
+		if (Binding.BindingType == NNX::EMLTensorBindingDataType::CPUMemory)
+		{
+			if (bIsInput) Tensor.SetUploadBuffer(Binding.CpuMemory);
+			else Tensor.SetDownloadBuffer(Binding.CpuMemory);
+		}
+		else if (Binding.BindingType == NNX::EMLTensorBindingDataType::RDGBuffer)
+		{
+			Tensor.SetBuffer(Binding.Buffer);
+		}
+		else
+		{
+			// Unsupported tensor binding type
+			return Idx;
+		}
+	}
+
+	return 0;
+}
+
+bool ApplyWeights(FRDGBuilder& GraphBuilder, TArray<FTensorHLSL>& OutTensorsHLSL, TArray<TRefCountPtr<FRDGPooledBuffer>> ExternalWeightsRDG)
+{
+	check(OutTensorsHLSL.Num() == ExternalWeightsRDG.Num());
+
+	for (int32 Idx = 0; Idx < ExternalWeightsRDG.Num(); ++Idx)
+	{
+		const TRefCountPtr<FRDGPooledBuffer>& PooledBuffer = ExternalWeightsRDG[Idx];
+		FTensorHLSL& Tensor = OutTensorsHLSL[Idx];
+		FRDGBufferRef Buffer = GraphBuilder.RegisterExternalBuffer(PooledBuffer);
+		Tensor.SetBuffer(Buffer);
+	}
+
+	return true;
+}
+
+}
+
+bool FModel::Init(TConstArrayView<uint8> ModelData)
+{
+	check(ModelData.Num() > 0);
+	FMLRuntimeFormat	Format;
+
+	if (!LoadModel(ModelData, Format))
+	{
+		return false;
+	}
+
+	// Create HLSL tensor and upload to GPU
+	PrepareWeights();
+
+	// Loop over all operators in the model and create them
+	for (int32 Idx = 0; Idx < Format.Operators.Num(); ++Idx)
+	{
+		const FString TypeName = Format.Operators[Idx].TypeName;
+
+		TArray<NNX::FTensorDesc> Inputs;
+		TArray<NNX::FTensorDesc> Outputs;
+		UE::NNECore::FAttributeMap AttributeMap;
+
+		for (int32 InputTensorIndex : Format.Operators[Idx].InTensors)
+		{
+			Inputs.Emplace(AllSymbolicTensorDescs[InputTensorIndex]);
+		}
+		for (int32 OutputTensorIndex : Format.Operators[Idx].OutTensors)
+		{
+			Outputs.Emplace(AllSymbolicTensorDescs[OutputTensorIndex]);
+		}
+		for (const FMLFormatAttributeDesc& Desc : Format.Operators[Idx].Attributes)
+		{
+			AttributeMap.SetAttribute(Desc.Name, Desc.Value);
+		}
+
+		FOperatorHlsl* Op = OpCreate(TypeName, Inputs, Outputs, AttributeMap);
+
+		if (!Op) //Op.Shader.IsNull())
+		{
+			UE_LOG(LogNNX, Warning, TEXT("Failed to create operator:%s"), *TypeName);
+
+			// TODO: Cleanup operators
+			return false;
+		}
+
+		Operators.Add(Op);
+	}
+
+	return true;
+}
+
+int FModel::SetInputTensorShapes(TConstArrayView<NNX::FTensorShape> InputShapes)
+{
+	int Res = FInferenceModelRDG::SetInputTensorShapes(InputShapes);
+	if (Res < 0) return Res;
+
+	auto ConvertTensors = [this] (TArray<FTensorHLSLRef> &TensorRefs, TArray<FTensorHLSL> &Tensors, const FTensorRDGArray &TensorRDGs, const TArray<int32>& TensorIndices)
+	{
+		Tensors.Empty();
+		Tensors.SetNum(TensorRDGs.Num());
+
+		for (int32 i = 0; i < TensorIndices.Num(); ++i)
+		{
+			const int32 Idx = TensorIndices[i];
+			Tensors[i] = FTensorHLSL(TensorRDGs[i]);
+			AllTensorHLSLRefs[Idx]= &Tensors[i];
+		}
+	};
+
+	AllTensorHLSLRefs.Empty();
+	AllTensorHLSLRefs.SetNumUninitialized(AllTensorRDGs.Num());
+	ConvertTensors(AllTensorHLSLRefs, InputTensorHLSLs, InputTensorRDGs, InputTensorIndices);
+	ConvertTensors(AllTensorHLSLRefs, OutputTensorHLSLs, OutputTensorRDGs, OutputTensorIndices);
+	ConvertTensors(AllTensorHLSLRefs, IntermediateTensorHLSLs, IntermediateTensorRDGs, IntermediateTensorIndices);
+	for (int32 i = 0; i < WeightTensorIndices.Num(); ++i)
+	{
+		const int32 Idx = WeightTensorIndices[i];
+		AllTensorHLSLRefs[Idx] = &WeightTensorHLSLs[i];
+	}
+
+	return 0;
+}
+
+int FModel::RunSync(TConstArrayView<NNX::FMLTensorBinding> InInputBindings, TConstArrayView<NNX::FMLTensorBinding> InOutputBindings)
+{
+	// Verify the model inputs were prepared
+	if (InputTensorShapes.Num() == 0)
+	{
+		UE_LOG(LogNNX, Error, TEXT("Run(): Input shapes are not set, please call SetInputTensorShapes."));
+		return -1;
+	}
+	
+	int Res = 0;
+	FEvent* Signal = FGenericPlatformProcess::GetSynchEventFromPool(false);
+
+	ENQUEUE_RENDER_COMMAND(FMLInferenceModel_Run)
+	(
+		[&Signal, &Res, this, InInputBindings, InOutputBindings](FRHICommandListImmediate& RHICmdList)
+		{
+			TOptional<ERHIPipeline>		Pipeline = RHICmdList.GetPipeline();
+
+			if (Pipeline == ERHIPipeline::None)
+			{
+				RHICmdList.SwitchPipeline(ERHIPipeline::Graphics);
+			}
+
+			FRDGBuilder	RDGBuilder(RHICmdList);
+
+			Res = EnqueueRDG(RDGBuilder, InInputBindings, InOutputBindings);
+			if (Res == 0)
+			{
+				RDGBuilder.Execute();
+
+				// FIXME: Using BlockUntilGPUIdle() prevents hang on Linux
+				// FIXME: Adapt to redesigned readback API (UE 5.2)
+				RHICmdList.BlockUntilGPUIdle();
+
+				for (FTensorHLSLRef Tensor : AllTensorHLSLRefs)
+				{
+					Tensor->Resolve();
+				}
+			}
+			
+			Signal->Trigger();
+		}
+	);
+
+	// We need to wait for render thread to finish
+	Signal->Wait();
+
+	FGenericPlatformProcess::ReturnSynchEventToPool(Signal);
+
+	return Res;
+}
+
+int FModel::EnqueueRDG(FRDGBuilder& GraphBuilder, TConstArrayView<NNX::FMLTensorBinding> InInputBindings, TConstArrayView<NNX::FMLTensorBinding> InOutputBindings)
+{
+	check(IsInRenderingThread());
+
+	int Res;
+
+	// Verify the model inputs were prepared
+	if (InputTensorShapes.Num() == 0)
+	{
+		UE_LOG(LogNNX, Error, TEXT("EnqueueRDG(): Input shapes are not set, please call SetInputTensorShapes."));
+		return -1;
+	}
+
+	Res = ApplyBinding(InputTensorHLSLs, InInputBindings, true);
+	if (Res != 0)
+	{
+		UE_LOG(LogNNX, Warning, TEXT("Invalid input tensor binding type for tensor index:%d"), Res);
+		return -1;
+	}
+
+	Res = ApplyBinding(OutputTensorHLSLs, InOutputBindings, false);
+	if (Res != 0)
+	{
+		UE_LOG(LogNNX, Warning, TEXT("Invalid output tensor binding type for tensor index:%d"), Res);
+		return -1;
+	}
+
+	if (!ApplyWeights(GraphBuilder, WeightTensorHLSLs, WeightsExternalRDGResources))
+	{
+		UE_LOG(LogNNX, Warning, TEXT("Could not register the weights for graph execution."));
+		return -1;
+	}
+
+	// Create RDG buffers where not yet present
+	for (FTensorHLSLRef Tensor : AllTensorHLSLRefs)
+	{
+		if (!Tensor->HasBuffer())
+		{
+			FRDGBufferDesc BufferDesc = FRDGBufferDesc::CreateBufferDesc(Tensor->GetElemByteSize(), Tensor->GetVolume());
+			if (Tensor->HasDownloadBuffer())
+			{
+				BufferDesc.Usage = EBufferUsageFlags(BufferDesc.Usage | BUF_SourceCopy);
+			}
+
+			const FRDGBufferRef TensorBuffer = GraphBuilder.CreateBuffer(BufferDesc, *Tensor->GetName(), ERDGBufferFlags::None);
+			
+			Tensor->SetBuffer(TensorBuffer);
+		}
+	}
+
+	for (FTensorHLSLRef Tensor : AllTensorHLSLRefs)
+	{
+		Tensor->EnqueueUploadRdg(GraphBuilder);
+	}
+
+	InternAddDispatchOps_RenderThread<FOperatorHlsl*>(GraphBuilder, AllTensorHLSLRefs, OperatorInputTensorIndices, OperatorOutputTensorIndices, Operators);
+
+	for (FTensorHLSLRef Tensor : AllTensorHLSLRefs)
+	{
+		Tensor->EnqueueDownloadRdg(GraphBuilder, bUseManualTransitions);
+	}
+
+	return 0;
+}
+
+int FModel::PrepareTensorShapesAndData()
+{
+	check(AllTensorRDGs.Num() == AllSymbolicTensorDescs.Num());
+	
+	if (Operators.Num() == 0)
+	{
+		UE_LOG(LogNNX, Warning, TEXT("No operators in model"));
+		return -1;
+	}
+
+	// Run model preparation (including shape inference) on all operators
+	// This loop could be abstracted to a different engine/system as it apply on FTensorRef & IPrepareOperator witch are RDG agnostics.
+	static constexpr int32 MaxExpectedInput = 10;
+	TArray<NNECore::Internal::FTensorRef, TInlineAllocator<MaxExpectedInput>> InputTensors;
+	TArray<NNECore::Internal::FTensorRef> OutputTensors;
+	TArray<bool> AllInitializedTensors;
+
+	checkCode(
+		AllInitializedTensors.Init(false, AllSymbolicTensorDescs.Num());
+		for (int32 Idx : InputTensorIndices)
+		{
+			AllInitializedTensors[Idx] = true;
+		}
+		for (int32 Idx : WeightTensorIndices)
+		{
+			AllInitializedTensors[Idx] = true;
+		}
+	);
+
+	// Run model preparation (including shape inference) on all operators
+	// This loop could be abstracted to a different system as it apply on FTensorRef & IPrepareOperator witch are RDG agnostics.
+	for (int32 Idx = 0; Idx < Operators.Num(); ++Idx)
+	{
+		InputTensors.Empty();
+		OutputTensors.Empty();
+
+		//Operator inputs
+		for (int32 i : OperatorInputTensorIndices[Idx])
+		{
+			checkf(AllInitializedTensors[i] == true, TEXT("Input tensor %d for operator %d should have been initialized."), i, Idx);
+			InputTensors.Emplace(AllTensorRDGs[i]);
+		}
+		//Operator outputs
+		for (int32 i : OperatorOutputTensorIndices[Idx])
+		{
+			OutputTensors.Emplace(AllTensorRDGs[i]);
+			checkf(AllInitializedTensors[i] == false, TEXT("Output tensor %d for operator %d should not have been initialized yet."), i, Idx);
+			checkCode(AllInitializedTensors[i] = true);
+		}
+
+		const FOperatorHlsl* Op = Operators[Idx];
+
+		if (Op->PrepareOutputs(InputTensors, OutputTensors) != 0)
+		{
+			//Operator could not prepare the output tensors, meaning we can't allocate
+			//output buffer before running the model. This engine does not support this.
+			UE_LOG(LogNNX, Warning, TEXT("Could not deduce tensor shapes for this model during shape inference, HLSL engine wont support the model as it need to precompute all shapes for performance reasons."));
+			AllTensorRDGs.Empty();
+			return -1;
+		}
+	}
+
+	checkCode(
+		for (int i = 0; i < AllInitializedTensors.Num(); ++i)
+		{
+			checkf(AllInitializedTensors[i], TEXT("Tensor at index %d, was not initialized by model preparation."));
+		};
+	);
+
+	return 0;
+}
+
+bool FModel::PrepareWeights()
+{
+	if (!WeightsExternalRDGResources.IsEmpty())
+	{
+		check(WeightsExternalRDGResources.Num() == WeightTensorHLSLs.Num())
+		return true;
+	}
+
+	//Convert to HLSL tensor
+	//This copy the weight. To be improved.
+	check(WeightTensorHLSLs.IsEmpty());
+	WeightTensorHLSLs.SetNum(WeightTensorRDGs.Num());
+	for (int32 Idx = 0; Idx < WeightTensorRDGs.Num(); ++Idx)
+	{
+		WeightTensorHLSLs[Idx] = FTensorHLSL(WeightTensorRDGs[Idx]);
+	}
+
+	//Upload to GPU
+	FEvent* Signal = FGenericPlatformProcess::GetSynchEventFromPool(false);
+	WeightsExternalRDGResources.SetNum(WeightTensorHLSLs.Num());
+
+	ENQUEUE_RENDER_COMMAND(FModel_PrepareWeights)
+	(
+		[&Signal, this](FRHICommandListImmediate& RHICmdList)
+		{
+			TOptional<ERHIPipeline> Pipeline = RHICmdList.GetPipeline();
+			if (Pipeline == ERHIPipeline::None)
+			{
+				RHICmdList.SwitchPipeline(ERHIPipeline::Graphics);
+			}
+
+			FRDGBuilder	RDGBuilder(RHICmdList);
+
+			for (int32 i = 0; i < WeightTensorHLSLs.Num(); ++i)
+			{
+				FTensorHLSL& Tensor = WeightTensorHLSLs[i];
+				check(!Tensor.HasBuffer());
+				check(Tensor.HasPreparedData());
+				FRDGBufferDesc BufferDesc = FRDGBufferDesc::CreateBufferDesc(Tensor.GetElemByteSize(), Tensor.GetVolume());
+				const FRDGBufferRef TransientRDGBuffer = RDGBuilder.CreateBuffer(BufferDesc, *Tensor.GetName(), ERDGBufferFlags::None);
+				WeightsExternalRDGResources[i] = RDGBuilder.ConvertToExternalBuffer(TransientRDGBuffer);
+				Tensor.SetBuffer(TransientRDGBuffer);
+				Tensor.EnqueueUploadRdg(RDGBuilder);
+			}
+
+			RDGBuilder.Execute();
+
+			//To prevent any problem if model is released before upload is done to the GPU. To be improved.
+			RHICmdList.BlockUntilGPUIdle();
+
+			Signal->Trigger();
+		}
+	);
+
+	// We need to wait for render thread to finish
+	Signal->Wait();
+
+	FGenericPlatformProcess::ReturnSynchEventToPool(Signal);
+
+	return true;
+}
+
+} // namespace UE::NNERuntimeRDG::Private::Hlsl
\ No newline at end of file
diff --git a/Engine/Plugins/Experimental/NNE/Source/NNERuntimeRDG/Private/NNERuntimeRDGModelHlsl.h b/Engine/Plugins/Experimental/NNE/Source/NNERuntimeRDG/Private/NNERuntimeRDGModelHlsl.h
new file mode 100644
--- /dev/null
+++ b/Engine/Plugins/Experimental/NNE/Source/NNERuntimeRDG/Private/NNERuntimeRDGModelHlsl.h
@@ -0,0 +1,55 @@
+// Copyright Epic Games, Inc. All Rights Reserved.
+
+#pragma once
+
+#include "NNERuntimeRDG.h"
+#include "NNECoreTypes.h"
+#include "NNERuntimeRDGTensorHlsl.h"
+
+class FRDGBuilder;
+
+namespace UE::NNERuntimeRDG::Private::Hlsl
+{
+struct FOperatorHlsl;
+
+class FModel : public FInferenceModelRDG
+{
+	
+public:
+
+	~FModel() = default;
+
+	bool Init(TConstArrayView<uint8> ModelData);
+
+	virtual int SetInputTensorShapes(TConstArrayView<NNECore::FTensorShape> InputShapes) override;
+
+	/**
+	 * Run the inference model (synchronous version)
+	 */
+	virtual int RunSync(TConstArrayView<NNX::FMLTensorBinding> InputBindings, TConstArrayView<NNX::FMLTensorBinding> OutputBindings) override;
+	
+	/**
+	 * Enqueue operators to RDG, the caller will run the GraphBuilder.Execute()
+	 */
+	virtual int EnqueueRDG(FRDGBuilder& GraphBuilder, TConstArrayView<NNX::FMLTensorBinding> InputBindings, TConstArrayView<NNX::FMLTensorBinding> OutputBindings) override;
+
+protected:
+
+	virtual void AddDispatchOps_RenderThread(FRDGBuilder& GraphBuilder) override { };
+	virtual int PrepareTensorShapesAndData() override;
+
+	bool PrepareWeights();
+
+private:
+
+	TArray<FOperatorHlsl*>	Operators;
+
+	TArray<FTensorHLSLRef> AllTensorHLSLRefs;
+	TArray<FTensorHLSL> InputTensorHLSLs;
+	TArray<FTensorHLSL> OutputTensorHLSLs;
+	TArray<FTensorHLSL> IntermediateTensorHLSLs;
+	TArray<FTensorHLSL> WeightTensorHLSLs;
+	TArray<TRefCountPtr<FRDGPooledBuffer>> WeightsExternalRDGResources;
+};
+
+} // namespace UE::NNERuntimeRDG::Private::Hlsl
\ No newline at end of file
diff --git a/Engine/Plugins/Experimental/NNE/Source/NNERuntimeRDG/Private/NNERuntimeRDGModule.cpp b/Engine/Plugins/Experimental/NNE/Source/NNERuntimeRDG/Private/NNERuntimeRDGModule.cpp
new file mode 100644
--- /dev/null
+++ b/Engine/Plugins/Experimental/NNE/Source/NNERuntimeRDG/Private/NNERuntimeRDGModule.cpp
@@ -0,0 +1,61 @@
+// Copyright Epic Games, Inc. All Rights Reserved.
+
+#include "NNXCore.h"
+#include "NNXRuntime.h"
+
+#include "NNERuntimeRDG.h"
+
+#include "Modules/ModuleManager.h"
+
+//
+//
+//
+class FNNXRuntimeRDGModule : public IModuleInterface
+{
+	NNX::IRuntime* DmlRuntime { nullptr };
+	NNX::IRuntime* HlslRuntime{ nullptr };
+
+public:
+
+	virtual void StartupModule() override
+	{
+#ifdef NNE_USE_DIRECTML
+		DmlRuntime = UE::NNERuntimeRDG::Private::Dml::FRuntimeDmlStartup();
+
+		if (DmlRuntime)
+		{
+			NNX::RegisterRuntime(DmlRuntime);
+		}
+#endif
+
+		HlslRuntime = UE::NNERuntimeRDG::Private::Hlsl::FRuntimeHlslStartup();
+
+		if (HlslRuntime)
+		{
+			NNX::RegisterRuntime(HlslRuntime);
+		}
+	}
+
+	virtual void ShutdownModule() override
+	{
+#ifdef NNE_USE_DIRECTML
+		if (DmlRuntime)
+		{
+			NNX::UnregisterRuntime(DmlRuntime);
+			DmlRuntime = nullptr;
+		}
+
+		UE::NNERuntimeRDG::Private::Dml::FRuntimeDmlShutdown();
+#endif
+
+		if (HlslRuntime)
+		{
+			NNX::UnregisterRuntime(HlslRuntime);
+			HlslRuntime = nullptr;
+		}
+
+		UE::NNERuntimeRDG::Private::Hlsl::FRuntimeHlslShutdown();
+	}
+};
+
+IMPLEMENT_MODULE(FNNXRuntimeRDGModule, NNXRuntimeRDG)
diff --git a/Engine/Plugins/Experimental/NNE/Source/NNERuntimeRDG/Private/NNERuntimeRDGNeuralPostProcessing.cpp b/Engine/Plugins/Experimental/NNE/Source/NNERuntimeRDG/Private/NNERuntimeRDGNeuralPostProcessing.cpp
new file mode 100644
--- /dev/null
+++ b/Engine/Plugins/Experimental/NNE/Source/NNERuntimeRDG/Private/NNERuntimeRDGNeuralPostProcessing.cpp
@@ -0,0 +1,433 @@
+// Copyright Epic Games, Inc. All Rights Reserved.
+
+#include "NNERuntimeRDGNeuralPostProcessing.h"
+#include "NNXCore.h"
+#include "PostProcess/PostProcessMaterial.h"
+#include "PostProcess/SceneRenderTargets.h"
+#include "ScreenPass.h"
+#include "PostProcess/PostProcessing.h"
+#include "PixelShaderUtils.h"
+
+#include "NNEHlslShadersNeuralPostProcessingCS.h"
+
+DECLARE_GPU_STAT_NAMED(FNNENeuralPostProcessingReadInput, TEXT("NNE.NeuralPostProcessing.ReadInput"));
+DECLARE_GPU_STAT_NAMED(FNNENeuralPostProcessingPreStep, TEXT("NNE.NeuralPostProcessing.PreStep"));
+DECLARE_GPU_STAT_NAMED(FNNENeuralPostProcessingPostStep, TEXT("NNE.NeuralPostProcessing.PostStep"));
+DECLARE_GPU_STAT_NAMED(FNNENeuralPostProcessingWriteOutput, TEXT("NNE.NeuralPostProcessing.WriteOutput"));
+
+FNNENeuralPostProcessing::FNNENeuralPostProcessing(const FAutoRegister& AutoRegister) : FSceneViewExtensionBase(AutoRegister) 
+{
+	NumEnabled = 0;
+	LastId = 0;
+}
+
+int32 FNNENeuralPostProcessing::Add(FString RuntimeName, UNNEModelData* ModelData)
+{
+	FScopeLock Lock(&CriticalSection);
+
+	// Create the model
+	NNX::IRuntime* Runtime = NNX::GetRuntime(RuntimeName);
+	if (!Runtime)
+	{
+		UE_LOG(LogNNX, Error, TEXT("FNNENeuralPostProcessing: No runtime '%s' found. Valid runtimes are: "), *RuntimeName);
+		TArray<NNX::IRuntime*> Runtimes = NNX::GetAllRuntimes();
+		for (int32 i = 0; i < Runtimes.Num(); i++)
+		{
+			UE_LOG(LogNNX, Error, TEXT("- %s"), *Runtimes[i]->GetRuntimeName());
+		}
+		return -1;
+	}
+
+	if (!ModelData)
+	{
+		UE_LOG(LogNNX, Error, TEXT("FNNENeuralPostProcessing: Valid model data required to load the model"));
+		return -1;
+	}
+
+	TConstArrayView<uint8> Data = ModelData->GetModelData(RuntimeName);
+	if (Data.Num() < 1)
+	{
+		UE_LOG(LogNNX, Error, TEXT("FNNENeuralPostProcessing: No model data for %s found"), *RuntimeName);
+		return -1;
+	}
+
+	TSharedPtr<NNX::FMLInferenceModel> Model = TSharedPtr<NNX::FMLInferenceModel>(Runtime->CreateModel(Data).Release());
+	if (!Model.IsValid())
+	{
+		UE_LOG(LogNNX, Error, TEXT("FNNENeuralPostProcessing: Could not create model using %s"), *RuntimeName);
+		return -1;
+	}
+
+	// Create a new id
+	LastId = (LastId + 1) > 0 ? (LastId + 1) : 1;
+
+	// Add the model to the map
+	Models.Add(LastId, Model);
+	SetWeight(LastId, 0.0);
+	SetRangeScale(LastId, 1.0);
+	SetInputSize(LastId, FIntPoint(-1, -1));
+
+	return LastId;
+}
+
+bool FNNENeuralPostProcessing::Remove(int32 ModelId)
+{
+	FScopeLock Lock(&CriticalSection);
+
+	bool bResult = Models.Find(ModelId) != nullptr;
+
+	Disable(ModelId);
+	Models.Remove(ModelId);
+	Weights.Remove(ModelId);
+	RangeScales.Remove(ModelId);
+	InputSizes.Remove(ModelId);
+
+	return bResult;
+}
+
+bool FNNENeuralPostProcessing::SetWeight(int32 ModelId, float Weight)
+{
+	FScopeLock Lock(&CriticalSection);
+
+	if (Models.Find(ModelId) != nullptr)
+	{
+		Weights.Add(ModelId, Weight);
+		return true;
+	}
+	return false;
+}
+
+bool FNNENeuralPostProcessing::SetRangeScale(int32 ModelId, float RangeScale)
+{
+	FScopeLock Lock(&CriticalSection);
+
+	if (Models.Find(ModelId) != nullptr)
+	{
+		RangeScales.Add(ModelId, RangeScale);
+		return true;
+	}
+	return false;
+}
+
+bool FNNENeuralPostProcessing::SetInputSize(int32 ModelId, FIntPoint InputSize)
+{
+	FScopeLock Lock(&CriticalSection);
+
+	if (Models.Find(ModelId) != nullptr)
+	{
+		InputSizes.Add(ModelId, InputSize);
+		return true;
+	}
+	return false;
+}
+
+void FNNENeuralPostProcessing::Enable(int32 ModelId)
+{
+	FScopeLock Lock(&CriticalSection);
+
+	if (Enabled.Find(ModelId) == nullptr)
+	{
+		Enabled.Add(ModelId);
+		NumEnabled++;
+	}
+}
+
+void FNNENeuralPostProcessing::Disable(int32 ModelId)
+{
+	FScopeLock Lock(&CriticalSection);
+
+	if (Enabled.Find(ModelId) != nullptr)
+	{
+		Enabled.Remove(ModelId);
+		NumEnabled--;
+	}
+}
+
+void FNNENeuralPostProcessing::PrePostProcessPass_RenderThread(FRDGBuilder& GraphBuilder, const FSceneView& View, const FPostProcessingInputs& Inputs)
+{
+	using namespace UE::NNEHlslShaders::Internal;
+	const float WeightEpsilon = 1.0/65536.0;
+
+	check(IsInRenderingThread());
+	check(View.bIsViewInfo);
+
+	float InputWeight = 0.0;
+	{
+		FScopeLock Lock(&CriticalSection);
+
+		if (NumEnabled < 1)
+		{
+			return;
+		}
+
+		float WeightSum = 0.0;
+		for (TPair<int32, float> Pair : Weights)
+		{
+			if (Enabled.Find(Pair.Key) != nullptr)
+			{
+				if (Pair.Value >= 0.0)
+				{
+					WeightSum += Pair.Value;
+				}
+				else
+				{
+					Weights.Add(Pair.Key, 0.0);
+				}
+			}
+		}
+
+		if (WeightSum > 1.0)
+		{
+			for (TPair<int32, float> Pair : Weights)
+			{
+				if (Enabled.Find(Pair.Key) != nullptr)
+				{
+					Weights.Add(Pair.Key, Pair.Value / WeightSum);
+				}
+			}
+		}
+		else
+		{
+			InputWeight = 1.0 - WeightSum;
+		}
+	}
+
+	const FIntRect Viewport = static_cast<const FViewInfo&>(View).ViewRect; 
+	FScreenPassTexture SceneColor((*Inputs.SceneTextures)->SceneColorTexture, Viewport);
+	FGlobalShaderMap* GlobalShaderMap = GetGlobalShaderMap(GMaxRHIFeatureLevel);
+	FIntPoint TextureSize = (*Inputs.SceneTextures)->SceneColorTexture->Desc.Extent;
+
+	// Read the input into an accumulation buffer 
+	FRDGBufferDesc AccumulationBufferDesc = FRDGBufferDesc::CreateBufferDesc(sizeof(float), TextureSize.X * TextureSize.Y * 3);
+	FRDGBufferRef AccumulationBuffer = GraphBuilder.CreateBuffer(AccumulationBufferDesc, *FString("NNENeuralPostProcessing::AccumulationBuffer"));
+	FRDGBufferUAVRef AccumulationBufferUAV = GraphBuilder.CreateUAV(FRDGBufferUAVDesc(AccumulationBuffer, PF_R32_FLOAT));
+
+	bool bOverwrite = true;
+	if (InputWeight > WeightEpsilon)
+	{
+		TNeuralPostProcessingReadInputCS::FParameters* ReadInputParameters = GraphBuilder.AllocParameters<TNeuralPostProcessingReadInputCS::FParameters>();
+		ReadInputParameters->InputTexture = (*Inputs.SceneTextures)->SceneColorTexture;
+		ReadInputParameters->InputTextureSampler = TStaticSamplerState<SF_Point, AM_Clamp, AM_Clamp, AM_Clamp>::GetRHI();
+		ReadInputParameters->InputTextureWidth = TextureSize.X;
+		ReadInputParameters->InputTextureHeight = TextureSize.Y;
+		ReadInputParameters->AccumulationBuffer = AccumulationBufferUAV;
+		ReadInputParameters->Weight = InputWeight;
+
+		FIntVector ReadInputThreadGroupCount = FIntVector(FMath::DivideAndRoundUp(TextureSize.X, FNeuralPostProcessingConstants::THREAD_GROUP_SIZE), FMath::DivideAndRoundUp(TextureSize.Y, FNeuralPostProcessingConstants::THREAD_GROUP_SIZE), 1);
+		TShaderMapRef<TNeuralPostProcessingReadInputCS> ReadInputShader(GlobalShaderMap);
+
+		RDG_EVENT_SCOPE(GraphBuilder, "NNE.NeuralPostProcessing.ReadInput");
+		RDG_GPU_STAT_SCOPE(GraphBuilder, FNNENeuralPostProcessingReadInput);
+
+		FComputeShaderUtils::AddPass(
+			GraphBuilder,
+			RDG_EVENT_NAME("NNE.NeuralPostProcessing.ReadInput"),
+			ERDGPassFlags::Compute | ERDGPassFlags::NeverCull,
+			ReadInputShader,
+			ReadInputParameters,
+			ReadInputThreadGroupCount);
+
+		bOverwrite = false;
+	}
+
+	// Process each network
+	{
+		FScopeLock Lock(&CriticalSection);
+		
+		for (TPair<int32, TSharedPtr<NNX::FMLInferenceModel, ESPMode::ThreadSafe>> Pair : Models)
+		{
+			if (Enabled.Find(Pair.Key) != nullptr && Weights.Find(Pair.Key) != nullptr)
+			{
+				float OutputWeight = *Weights.Find(Pair.Key);
+				if (OutputWeight > WeightEpsilon)
+				{
+					float RangeScale = 1.0;
+					if (RangeScales.Find(Pair.Key))
+					{
+						RangeScale = *RangeScales.Find(Pair.Key);
+					}
+
+					FIntPoint InputSize = FIntPoint(-1, -1);
+					if (InputSizes.Find(Pair.Key))
+					{
+						InputSize = *InputSizes.Find(Pair.Key);
+					}
+
+					UE::NNECore::FSymbolicTensorShape InputShape = Pair.Value->GetInputTensorDescs()[0].GetShape();
+
+					checkf(InputShape.Rank() == 4, TEXT("Neural Post Processing requires models with input shape 1 x 3 x height x width!"))
+					checkf(InputShape.GetData()[0] == 1, TEXT("Neural Post Processing requires models with input shape 1 x 3 x height x width!"))
+					checkf(InputShape.GetData()[1] == 3, TEXT("Neural Post Processing requires models with input shape 1 x 3 x height x width!"))
+
+					int32 NeuralNetworkInputWidth = InputShape.GetData()[3] >= 0 ? InputShape.GetData()[3] : (InputSize.X >= 0 ? InputSize.X : TextureSize.X);
+					int32 NeuralNetworkInputHeight = InputShape.GetData()[2] >= 0 ? InputShape.GetData()[2] : (InputSize.Y >= 0 ? InputSize.Y : TextureSize.Y);
+
+					FRDGBufferDesc InputBufferDesc = FRDGBufferDesc::CreateBufferDesc(sizeof(float), NeuralNetworkInputWidth * NeuralNetworkInputHeight * 3);
+					FRDGBufferRef InputBuffer = GraphBuilder.CreateBuffer(InputBufferDesc, *(FString("NNENeuralPostProcessing::NeuralNetowrkInput_") + FString::FromInt(Pair.Key)));
+					FRDGBufferUAVRef InputBufferUAV = GraphBuilder.CreateUAV(FRDGBufferUAVDesc(InputBuffer, PF_R32_FLOAT));
+
+					TNeuralPostProcessingPreStepCS::FParameters* PreStepParameters = GraphBuilder.AllocParameters<TNeuralPostProcessingPreStepCS::FParameters>();
+					PreStepParameters->InputTexture = (*Inputs.SceneTextures)->SceneColorTexture;
+					PreStepParameters->InputTextureSampler = TStaticSamplerState<SF_Point, AM_Clamp, AM_Clamp, AM_Clamp>::GetRHI();
+					PreStepParameters->InputTextureWidth = TextureSize.X;
+					PreStepParameters->InputTextureHeight = TextureSize.Y;
+					PreStepParameters->InputBufferWidth = NeuralNetworkInputWidth;
+					PreStepParameters->InputBufferHeight = NeuralNetworkInputHeight;
+					PreStepParameters->InputBuffer = InputBufferUAV;
+					PreStepParameters->RangeScale = RangeScale;
+
+					FIntVector PreStepThreadGroupCount = FIntVector(FMath::DivideAndRoundUp(NeuralNetworkInputWidth, FNeuralPostProcessingConstants::THREAD_GROUP_SIZE), FMath::DivideAndRoundUp(NeuralNetworkInputHeight, FNeuralPostProcessingConstants::THREAD_GROUP_SIZE), 1);
+					TShaderMapRef<TNeuralPostProcessingPreStepCS> PreStepShader(GlobalShaderMap);
+
+					RDG_EVENT_SCOPE(GraphBuilder, "NNE.NeuralPostProcessing.PreStep");
+					RDG_GPU_STAT_SCOPE(GraphBuilder, FNNENeuralPostProcessingPreStep);
+
+					FComputeShaderUtils::AddPass(
+						GraphBuilder,
+						RDG_EVENT_NAME("NNE.NeuralPostProcessing.PreStep"),
+						ERDGPassFlags::Compute | ERDGPassFlags::NeverCull,
+						PreStepShader,
+						PreStepParameters,
+						PreStepThreadGroupCount);
+
+					TArray<NNX::FTensorShape> InputShapes;
+					InputShapes.Add(NNX::FTensorShape::Make({ 1, 3, (uint32)NeuralNetworkInputHeight, (uint32)NeuralNetworkInputWidth }));
+
+					Pair.Value->SetInputTensorShapes(InputShapes);
+					NNX::FTensorShape OutputShape = Pair.Value->GetOutputTensorShapes()[0];
+
+					checkf(OutputShape.Rank() == 4, TEXT("Neural Post Processing requires models with output shape 1 x 3 x height x width!"))
+					checkf(OutputShape.GetData()[0] == 1, TEXT("Neural Post Processing requires models with output shape 1 x 3 x height x width!"))
+					checkf(OutputShape.GetData()[1] == 3, TEXT("Neural Post Processing requires models with output shape 1 x 3 x height x width!"))
+					checkf(OutputShape.GetData()[2] > 0, TEXT("Neural Post Processing requires models with output height > 0!"))
+					checkf(OutputShape.GetData()[3] > 0, TEXT("Neural Post Processing requires models with output width > 0!"))
+
+					int32 NeuralNetworkOutputWidth = OutputShape.GetData()[3];
+					int32 NeuralNetworkOutputHeight = OutputShape.GetData()[2];
+
+					FRDGBufferDesc OutputBufferDesc = FRDGBufferDesc::CreateBufferDesc(sizeof(float), NeuralNetworkOutputWidth * NeuralNetworkOutputHeight * 3);
+					FRDGBufferRef OutputBuffer = GraphBuilder.CreateBuffer(OutputBufferDesc, *(FString("NNENeuralPostProcessing::NeuralNetowrkOutput_") + FString::FromInt(Pair.Key)));
+					FRDGBufferUAVRef OutputBufferUAV = GraphBuilder.CreateUAV(FRDGBufferUAVDesc(OutputBuffer, PF_R32_FLOAT));
+
+					TArray<NNX::FMLTensorBinding> InputBindings;
+					InputBindings.Add(NNX::FMLTensorBinding::FromRDG(InputBuffer, InputBufferDesc.NumElements * InputBufferDesc.BytesPerElement, 0));
+					TArray<NNX::FMLTensorBinding> OutputBindings;
+					OutputBindings.Add(NNX::FMLTensorBinding::FromRDG(OutputBuffer, OutputBufferDesc.NumElements * OutputBufferDesc.BytesPerElement, 0));
+					Pair.Value->EnqueueRDG(GraphBuilder, InputBindings, OutputBindings);
+
+					TNeuralPostProcessingPostStepCS::FParameters* PostStepParameters = GraphBuilder.AllocParameters<TNeuralPostProcessingPostStepCS::FParameters>();
+					PostStepParameters->OutputBufferWidth = NeuralNetworkOutputWidth;
+					PostStepParameters->OutputBufferHeight = NeuralNetworkOutputHeight;
+					PostStepParameters->OutputBuffer = OutputBufferUAV;
+					PostStepParameters->InputTextureWidth = TextureSize.X;
+					PostStepParameters->InputTextureHeight = TextureSize.Y;
+					PostStepParameters->AccumulationBuffer = AccumulationBufferUAV;
+					PostStepParameters->Weight = OutputWeight;
+					PostStepParameters->RangeScale = RangeScale;
+
+					TNeuralPostProcessingPostStepCS::FPermutationDomain PermutationVector;
+					PermutationVector.Set<TNeuralPostProcessingPostStepCS::FNeuralPostProcessingOverwrite>(bOverwrite ? ENeuralPostProcessingOverwrite::Yes : ENeuralPostProcessingOverwrite::No);
+					PermutationVector.Set<TNeuralPostProcessingPostStepCS::FNeuralPostProcessingInterpolate>((NeuralNetworkOutputWidth == TextureSize.X && NeuralNetworkOutputHeight == TextureSize.Y) ? ENeuralPostProcessingInterpolate::No : ENeuralPostProcessingInterpolate::Yes);
+
+					FIntVector PostStepThreadGroupCount = FIntVector(FMath::DivideAndRoundUp(TextureSize.X, FNeuralPostProcessingConstants::THREAD_GROUP_SIZE), FMath::DivideAndRoundUp(TextureSize.Y, FNeuralPostProcessingConstants::THREAD_GROUP_SIZE), 1);
+					TShaderMapRef<TNeuralPostProcessingPostStepCS> PostStepShader(GlobalShaderMap, PermutationVector);
+
+					RDG_EVENT_SCOPE(GraphBuilder, "NNE.NeuralPostProcessing.PostStep");
+					RDG_GPU_STAT_SCOPE(GraphBuilder, FNNENeuralPostProcessingPostStep);
+
+					FComputeShaderUtils::AddPass(
+						GraphBuilder,
+						RDG_EVENT_NAME("NNE.NeuralPostProcessing.PostStep"),
+						ERDGPassFlags::Compute | ERDGPassFlags::NeverCull,
+						PostStepShader,
+						PostStepParameters,
+						PostStepThreadGroupCount);
+
+					bOverwrite = false;
+				}
+			}
+		}
+	}
+
+	// Write the result
+	TNeuralPostProcessingWriteOutputPS::FParameters* WriteOutputParameters = GraphBuilder.AllocParameters<TNeuralPostProcessingWriteOutputPS::FParameters>();
+	WriteOutputParameters->AccumulationBuffer = AccumulationBufferUAV;
+	WriteOutputParameters->InputTextureWidth = TextureSize.X;
+	WriteOutputParameters->InputTextureHeight = TextureSize.Y;
+	WriteOutputParameters->RenderTargets[0] = FRenderTargetBinding(SceneColor.Texture, ERenderTargetLoadAction::ENoAction);
+
+	RDG_EVENT_SCOPE(GraphBuilder, "NNE.NeuralPostProcessing.WriteOutput");
+	RDG_GPU_STAT_SCOPE(GraphBuilder, FNNENeuralPostProcessingWriteOutput);
+
+	TShaderMapRef<TNeuralPostProcessingWriteOutputPS> WriteOutputShader(GlobalShaderMap);
+	FPixelShaderUtils::AddFullscreenPass(
+		GraphBuilder,
+		GlobalShaderMap,
+		RDG_EVENT_NAME("NNE.NeuralPostProcessing.WriteOutput"),
+		WriteOutputShader,
+		WriteOutputParameters,
+		Viewport);
+}
+
+int32 UNNENeuralPostProcessing::Add(FString RuntimeName, UNNEModelData* ModelData)
+{
+	if (!NeuralPostProcessing.IsValid())
+	{
+		NeuralPostProcessing = FSceneViewExtensions::NewExtension<FNNENeuralPostProcessing>();
+	}
+	return NeuralPostProcessing->Add(RuntimeName, ModelData);
+}
+
+bool UNNENeuralPostProcessing::Remove(int32 ModelId)
+{
+	if (!NeuralPostProcessing.IsValid())
+	{
+		NeuralPostProcessing = FSceneViewExtensions::NewExtension<FNNENeuralPostProcessing>();
+	}
+	return NeuralPostProcessing->Remove(ModelId);
+}
+
+bool UNNENeuralPostProcessing::SetWeight(int32 ModelId, float Weight)
+{
+	if (!NeuralPostProcessing.IsValid())
+	{
+		NeuralPostProcessing = FSceneViewExtensions::NewExtension<FNNENeuralPostProcessing>();
+	}
+	return NeuralPostProcessing->SetWeight(ModelId, Weight);
+}
+
+bool UNNENeuralPostProcessing::SetRangeScale(int32 ModelId, float RangeScale)
+{
+	if (!NeuralPostProcessing.IsValid())
+	{
+		NeuralPostProcessing = FSceneViewExtensions::NewExtension<FNNENeuralPostProcessing>();
+	}
+	return NeuralPostProcessing->SetRangeScale(ModelId, RangeScale);
+}
+
+bool UNNENeuralPostProcessing::SetInputSize(int32 ModelId, FIntPoint InputSize)
+{
+	if (!NeuralPostProcessing.IsValid())
+	{
+		NeuralPostProcessing = FSceneViewExtensions::NewExtension<FNNENeuralPostProcessing>();
+	}
+	return NeuralPostProcessing->SetInputSize(ModelId, InputSize);
+}
+
+void UNNENeuralPostProcessing::Enable(int32 ModelId)
+{
+	if (!NeuralPostProcessing.IsValid())
+	{
+		NeuralPostProcessing = FSceneViewExtensions::NewExtension<FNNENeuralPostProcessing>();
+	}
+	NeuralPostProcessing->Enable(ModelId);
+}
+
+void UNNENeuralPostProcessing::Disable(int32 ModelId)
+{
+	if (!NeuralPostProcessing.IsValid())
+	{
+		NeuralPostProcessing = FSceneViewExtensions::NewExtension<FNNENeuralPostProcessing>();
+	}
+	NeuralPostProcessing->Disable(ModelId);
+}
\ No newline at end of file
diff --git a/Engine/Plugins/Experimental/NNE/Source/NNERuntimeRDG/Private/NNERuntimeRDGTensorHlsl.cpp b/Engine/Plugins/Experimental/NNE/Source/NNERuntimeRDG/Private/NNERuntimeRDGTensorHlsl.cpp
new file mode 100644
--- /dev/null
+++ b/Engine/Plugins/Experimental/NNE/Source/NNERuntimeRDG/Private/NNERuntimeRDGTensorHlsl.cpp
@@ -0,0 +1,69 @@
+// Copyright Epic Games, Inc. All Rights Reserved.
+
+#include "NNERuntimeRDGTensorHlsl.h"
+#include "RenderGraphBuilder.h"
+
+namespace UE::NNERuntimeRDG::Private::Hlsl
+{
+
+void FTensorHLSL::EnqueueUploadRdg(FRDGBuilder& GraphBuilder)
+{
+	if (!HasPreparedData() && !HasUploadBuffer()) return;
+	
+	check(HasPreparedData() != HasUploadBuffer());
+	check(HasBuffer());
+
+	FCpuBufferRef CPUBuffer = HasUploadBuffer() ? GetUploadBuffer() : PreparedData.GetData();
+	GraphBuilder.QueueBufferUpload(GetBuffer(), CPUBuffer, GetDataSize(), ERDGInitialDataFlags::NoCopy);
+}
+
+void FTensorHLSL::EnqueueDownloadRdg(FRDGBuilder& GraphBuilder, bool bUseManualTransitions)
+{
+	if (!HasDownloadBuffer()) return;
+
+	check(HasBuffer());
+	check(!Readback);
+
+	Readback = MakeUnique<FRHIGPUBufferReadback>(FName(TEXT("FMLTensorReadback_") + GetName()));
+
+	FMLTensorReadbackParameters* TensorReadbackParams = GraphBuilder.AllocParameters<FMLTensorReadbackParameters>();
+	TensorReadbackParams->Buffer = GetBuffer();
+
+	GraphBuilder.AddPass(
+		RDG_EVENT_NAME("FMLInferenceModelAddTensorReadback:%s", *GetName()),
+		TensorReadbackParams,
+		ERDGPassFlags::Readback | ERDGPassFlags::NeverCull,
+		[this, bUseManualTransitions, TensorReadbackParams](FRHICommandListImmediate& RHICmdList)
+		{
+			FRHIBuffer* OutputBuffer = TensorReadbackParams->Buffer->GetRHI();
+
+			// TODO: FIXME: We need to transition the resources for DirectML
+			if (bUseManualTransitions)
+			{
+				FRHITransitionInfo Transitions[] =
+				{
+					FRHITransitionInfo(OutputBuffer, ERHIAccess::UAVCompute, ERHIAccess::CopySrc)
+				};
+
+				RHICmdList.Transition(MakeArrayView(Transitions, UE_ARRAY_COUNT(Transitions)));
+				RHICmdList.SubmitCommandsHint();
+			}
+
+			Readback->EnqueueCopy(RHICmdList, OutputBuffer, GetDataSize());
+		}
+	);
+}
+
+void FTensorHLSL::Resolve()
+{
+	if (!Readback) return;
+
+	const void* BuffData = Readback->Lock(GetDataSize());
+	check(BuffData);
+
+	FMemory::Memcpy(GetDownloadBuffer(), BuffData, GetDataSize());
+
+	Readback->Unlock();
+}
+
+} // namespace UE::NNERuntimeRDG::Private::Hlsl
\ No newline at end of file
diff --git a/Engine/Plugins/Experimental/NNE/Source/NNERuntimeRDG/Private/NNERuntimeRDGTensorHlsl.h b/Engine/Plugins/Experimental/NNE/Source/NNERuntimeRDG/Private/NNERuntimeRDGTensorHlsl.h
new file mode 100644
--- /dev/null
+++ b/Engine/Plugins/Experimental/NNE/Source/NNERuntimeRDG/Private/NNERuntimeRDGTensorHlsl.h
@@ -0,0 +1,61 @@
+// Copyright Epic Games, Inc. All Rights Reserved.
+
+#pragma once
+
+#include "NNERuntimeRDG.h"
+#include "NNECoreTypes.h"
+#include "RHIGPUReadback.h"
+
+using FCpuBufferRef = void*;
+
+class FRDGBuilder;
+
+namespace UE::NNERuntimeRDG::Private::Hlsl
+{
+
+class FTensorHLSL : public FTensorRDG
+{
+
+private:
+
+	FCpuBufferRef UploadBuffer{};
+	FCpuBufferRef DownloadBuffer{};
+
+	TUniquePtr<FRHIGPUBufferReadback> Readback;
+
+public:
+
+	FTensorHLSL() = default;
+
+	FTensorHLSL(const FTensorRDG &TensorRDG)
+		: FTensorHLSL(TensorRDG.GetName(), TensorRDG.GetDataType(), TensorRDG.GetShape())
+	{
+		this->PreparedData = TensorRDG.GetPreparedData<uint8>();//This copy the data, we should find another way
+	}
+	
+	FTensorHLSL(const FString& Name, ENNETensorDataType DataType, const NNX::FTensorShape& Shape)
+	{
+		this->Name = Name;
+		this->DataType = DataType;
+		this->Shape = Shape;
+		Volume = Shape.Volume();
+		DataSize = (uint64)NNECore::GetTensorDataTypeSizeInBytes(DataType) * Volume;
+	}
+
+	bool HasUploadBuffer() const { return UploadBuffer != FCpuBufferRef{}; }
+	void SetUploadBuffer(FCpuBufferRef Inbuffer){ check(!HasPreparedData()); UploadBuffer = Inbuffer; }
+	FCpuBufferRef GetUploadBuffer() const { return UploadBuffer; }
+
+	bool HasDownloadBuffer() const { return DownloadBuffer != FCpuBufferRef{}; }
+	void SetDownloadBuffer(FCpuBufferRef Inbuffer){ DownloadBuffer = Inbuffer; }
+	FCpuBufferRef GetDownloadBuffer() const { return DownloadBuffer; }
+
+	void EnqueueUploadRdg(FRDGBuilder& GraphBuilder);
+	void EnqueueDownloadRdg(FRDGBuilder& GraphBuilder, bool bUseManualTransitions);
+
+	void Resolve();
+};
+
+using FTensorHLSLRef = FTensorHLSL*;
+
+} // namespace UE::NNERuntimeRDG::Private::Hlsl
\ No newline at end of file
diff --git a/Engine/Plugins/Experimental/NNE/Source/NNXCore/Internal/NNECoreOperator.h b/Engine/Plugins/Experimental/NNE/Source/NNXCore/Internal/NNECoreOperator.h
new file mode 100644
--- /dev/null
+++ b/Engine/Plugins/Experimental/NNE/Source/NNXCore/Internal/NNECoreOperator.h
@@ -0,0 +1,88 @@
+// Copyright Epic Games, Inc. All Rights Reserved.
+
+#pragma once
+
+namespace UE::NNECore::Internal
+{
+
+	//One input element-wise operators
+	enum class EElementWiseUnaryOperatorType : uint8
+	{
+		Abs,
+		Acos,
+		Acosh,
+		Asin,
+		Asinh,
+		Atan,
+		Atanh,
+		//BitShift, //Note: need integer tensors
+		//Cast,     //Note: ability to cast tensor type
+		Ceil,
+		//Clip,     //Note: need scalar tensor inputs
+		Cos,
+		Cosh,
+		Elu,
+		Erf,
+		Exp,
+		Floor,
+		IsInf,
+		IsNan,
+		HardSigmoid,
+		HardSwish,
+		LeakyRelu,
+		Log,
+		Neg,
+		//Not,      //Note: need bool tensors
+		Reciprocal,
+		Relu,
+		Round,
+		Selu,
+		Sigmoid,
+		Sign,
+		Sin,
+		Sinh,
+		Softplus,
+		Softsign,
+		Sqrt,
+		Tan,
+		Tanh,
+
+		MAX
+	};
+
+	//Two inputs element-wise operators with multi-directional broadcast
+	//see [REDACTED_PROFILE_URL]/onnx/blob/main/docs/Broadcasting.md
+	enum class EElementWiseBinaryOperatorType : uint8
+	{
+		Add,
+		//And,           //Note: need boolean tensors
+		Div,
+		//Equal,         //Note: need boolean tensors
+		//Greater,       //Note: need boolean tensors
+		//GreaterOrEqual,//Note: need boolean tensors
+		//Less,          //Note: need boolean tensors
+		//LessOrEqual,   //Note: need boolean tensors
+		Mod,
+		Mul,
+		//Or,            //Note: need boolean tensors
+		Prelu,           //Note: only broadcast from slope to input.
+		Pow,
+		Sub,
+		//Xor,           //Note:  need boolean tensors
+
+		MAX
+	};
+
+	//Variable number of inputs element-wise operators with multi-directional broadcast
+	//see [REDACTED_PROFILE_URL]/onnx/blob/main/docs/Broadcasting.md
+	enum class EElementWiseVariadicOperatorType : uint8
+	{
+		Max,
+		Min,
+		Mean,
+		Sum,
+
+		MAX
+	};
+
+} // namespace UE::NNECore::Internal
\ No newline at end of file
diff --git a/Engine/Plugins/Experimental/NNE/Source/NNXCore/Public/NNXOperator.h b/Engine/Plugins/Experimental/NNE/Source/NNXCore/Public/NNXOperator.h
deleted file mode 100644
diff --git a/Engine/Plugins/Experimental/NNE/Source/NNXRuntimeRDG/NNXRuntimeRDG.Build.cs b/Engine/Plugins/Experimental/NNE/Source/NNXRuntimeRDG/NNXRuntimeRDG.Build.cs
deleted file mode 100644
diff --git a/Engine/Plugins/Experimental/NNE/Source/NNXRuntimeRDG/Private/Hlsl/NNIRuntimeRDGConv.cpp b/Engine/Plugins/Experimental/NNE/Source/NNXRuntimeRDG/Private/Hlsl/NNIRuntimeRDGConv.cpp
deleted file mode 100644
diff --git a/Engine/Plugins/Experimental/NNE/Source/NNXRuntimeRDG/Private/Hlsl/NNIRuntimeRDGConv.h b/Engine/Plugins/Experimental/NNE/Source/NNXRuntimeRDG/Private/Hlsl/NNIRuntimeRDGConv.h
deleted file mode 100644
diff --git a/Engine/Plugins/Experimental/NNE/Source/NNXRuntimeRDG/Private/Hlsl/NNIRuntimeRDGConvTranspose.cpp b/Engine/Plugins/Experimental/NNE/Source/NNXRuntimeRDG/Private/Hlsl/NNIRuntimeRDGConvTranspose.cpp
deleted file mode 100644
diff --git a/Engine/Plugins/Experimental/NNE/Source/NNXRuntimeRDG/Private/Hlsl/NNIRuntimeRDGConvTranspose.h b/Engine/Plugins/Experimental/NNE/Source/NNXRuntimeRDG/Private/Hlsl/NNIRuntimeRDGConvTranspose.h
deleted file mode 100644
diff --git a/Engine/Plugins/Experimental/NNE/Source/NNXRuntimeRDG/Private/Hlsl/NNIRuntimeRDGElementWiseBinary.cpp b/Engine/Plugins/Experimental/NNE/Source/NNXRuntimeRDG/Private/Hlsl/NNIRuntimeRDGElementWiseBinary.cpp
deleted file mode 100644
diff --git a/Engine/Plugins/Experimental/NNE/Source/NNXRuntimeRDG/Private/Hlsl/NNIRuntimeRDGElementWiseBinary.h b/Engine/Plugins/Experimental/NNE/Source/NNXRuntimeRDG/Private/Hlsl/NNIRuntimeRDGElementWiseBinary.h
deleted file mode 100644
diff --git a/Engine/Plugins/Experimental/NNE/Source/NNXRuntimeRDG/Private/Hlsl/NNIRuntimeRDGElementWiseUnary.cpp b/Engine/Plugins/Experimental/NNE/Source/NNXRuntimeRDG/Private/Hlsl/NNIRuntimeRDGElementWiseUnary.cpp
deleted file mode 100644
diff --git a/Engine/Plugins/Experimental/NNE/Source/NNXRuntimeRDG/Private/Hlsl/NNIRuntimeRDGElementWiseUnary.h b/Engine/Plugins/Experimental/NNE/Source/NNXRuntimeRDG/Private/Hlsl/NNIRuntimeRDGElementWiseUnary.h
deleted file mode 100644
diff --git a/Engine/Plugins/Experimental/NNE/Source/NNXRuntimeRDG/Private/Hlsl/NNIRuntimeRDGElementWiseVariadic.cpp b/Engine/Plugins/Experimental/NNE/Source/NNXRuntimeRDG/Private/Hlsl/NNIRuntimeRDGElementWiseVariadic.cpp
deleted file mode 100644
diff --git a/Engine/Plugins/Experimental/NNE/Source/NNXRuntimeRDG/Private/Hlsl/NNIRuntimeRDGElementWiseVariadic.h b/Engine/Plugins/Experimental/NNE/Source/NNXRuntimeRDG/Private/Hlsl/NNIRuntimeRDGElementWiseVariadic.h
deleted file mode 100644
diff --git a/Engine/Plugins/Experimental/NNE/Source/NNXRuntimeRDG/Private/Hlsl/NNIRuntimeRDGGather.cpp b/Engine/Plugins/Experimental/NNE/Source/NNXRuntimeRDG/Private/Hlsl/NNIRuntimeRDGGather.cpp
deleted file mode 100644
diff --git a/Engine/Plugins/Experimental/NNE/Source/NNXRuntimeRDG/Private/Hlsl/NNIRuntimeRDGGather.h b/Engine/Plugins/Experimental/NNE/Source/NNXRuntimeRDG/Private/Hlsl/NNIRuntimeRDGGather.h
deleted file mode 100644
diff --git a/Engine/Plugins/Experimental/NNE/Source/NNXRuntimeRDG/Private/Hlsl/NNIRuntimeRDGGemm.cpp b/Engine/Plugins/Experimental/NNE/Source/NNXRuntimeRDG/Private/Hlsl/NNIRuntimeRDGGemm.cpp
deleted file mode 100644
diff --git a/Engine/Plugins/Experimental/NNE/Source/NNXRuntimeRDG/Private/Hlsl/NNIRuntimeRDGGemm.h b/Engine/Plugins/Experimental/NNE/Source/NNXRuntimeRDG/Private/Hlsl/NNIRuntimeRDGGemm.h
deleted file mode 100644
diff --git a/Engine/Plugins/Experimental/NNE/Source/NNXRuntimeRDG/Private/Hlsl/NNIRuntimeRDGInstanceNormalization.cpp b/Engine/Plugins/Experimental/NNE/Source/NNXRuntimeRDG/Private/Hlsl/NNIRuntimeRDGInstanceNormalization.cpp
deleted file mode 100644
diff --git a/Engine/Plugins/Experimental/NNE/Source/NNXRuntimeRDG/Private/Hlsl/NNIRuntimeRDGInstanceNormalization.h b/Engine/Plugins/Experimental/NNE/Source/NNXRuntimeRDG/Private/Hlsl/NNIRuntimeRDGInstanceNormalization.h
deleted file mode 100644
diff --git a/Engine/Plugins/Experimental/NNE/Source/NNXRuntimeRDG/Private/Hlsl/NNIRuntimeRDGMatMul.cpp b/Engine/Plugins/Experimental/NNE/Source/NNXRuntimeRDG/Private/Hlsl/NNIRuntimeRDGMatMul.cpp
deleted file mode 100644
diff --git a/Engine/Plugins/Experimental/NNE/Source/NNXRuntimeRDG/Private/Hlsl/NNIRuntimeRDGMatMul.h b/Engine/Plugins/Experimental/NNE/Source/NNXRuntimeRDG/Private/Hlsl/NNIRuntimeRDGMatMul.h
deleted file mode 100644
diff --git a/Engine/Plugins/Experimental/NNE/Source/NNXRuntimeRDG/Private/Hlsl/NNIRuntimeRDGPad.cpp b/Engine/Plugins/Experimental/NNE/Source/NNXRuntimeRDG/Private/Hlsl/NNIRuntimeRDGPad.cpp
deleted file mode 100644
diff --git a/Engine/Plugins/Experimental/NNE/Source/NNXRuntimeRDG/Private/Hlsl/NNIRuntimeRDGPad.h b/Engine/Plugins/Experimental/NNE/Source/NNXRuntimeRDG/Private/Hlsl/NNIRuntimeRDGPad.h
deleted file mode 100644
diff --git a/Engine/Plugins/Experimental/NNE/Source/NNXRuntimeRDG/Private/Hlsl/NNIRuntimeRDGUpsample.cpp b/Engine/Plugins/Experimental/NNE/Source/NNXRuntimeRDG/Private/Hlsl/NNIRuntimeRDGUpsample.cpp
deleted file mode 100644
diff --git a/Engine/Plugins/Experimental/NNE/Source/NNXRuntimeRDG/Private/Hlsl/NNIRuntimeRDGUpsample.h b/Engine/Plugins/Experimental/NNE/Source/NNXRuntimeRDG/Private/Hlsl/NNIRuntimeRDGUpsample.h
deleted file mode 100644
diff --git a/Engine/Plugins/Experimental/NNE/Source/NNXRuntimeRDG/Private/NNERuntimeRDGModelHLSL.cpp b/Engine/Plugins/Experimental/NNE/Source/NNXRuntimeRDG/Private/NNERuntimeRDGModelHLSL.cpp
deleted file mode 100644
diff --git a/Engine/Plugins/Experimental/NNE/Source/NNXRuntimeRDG/Private/NNERuntimeRDGModelHLSL.h b/Engine/Plugins/Experimental/NNE/Source/NNXRuntimeRDG/Private/NNERuntimeRDGModelHLSL.h
deleted file mode 100644
diff --git a/Engine/Plugins/Experimental/NNE/Source/NNXRuntimeRDG/Private/NNERuntimeRDGNeuralPostProcessing.cpp b/Engine/Plugins/Experimental/NNE/Source/NNXRuntimeRDG/Private/NNERuntimeRDGNeuralPostProcessing.cpp
deleted file mode 100644
diff --git a/Engine/Plugins/Experimental/NNE/Source/NNXRuntimeRDG/Private/NNERuntimeRDGTensorHLSL.cpp b/Engine/Plugins/Experimental/NNE/Source/NNXRuntimeRDG/Private/NNERuntimeRDGTensorHLSL.cpp
deleted file mode 100644
diff --git a/Engine/Plugins/Experimental/NNE/Source/NNXRuntimeRDG/Private/NNERuntimeRDGTensorHLSL.h b/Engine/Plugins/Experimental/NNE/Source/NNXRuntimeRDG/Private/NNERuntimeRDGTensorHLSL.h
deleted file mode 100644
diff --git a/Engine/Plugins/Experimental/NNE/Source/NNXRuntimeRDG/Private/NNXRuntimeDML.cpp b/Engine/Plugins/Experimental/NNE/Source/NNXRuntimeRDG/Private/NNXRuntimeDML.cpp
deleted file mode 100644
diff --git a/Engine/Plugins/Experimental/NNE/Source/NNXRuntimeRDG/Private/NNXRuntimeHLSL.cpp b/Engine/Plugins/Experimental/NNE/Source/NNXRuntimeRDG/Private/NNXRuntimeHLSL.cpp
deleted file mode 100644
diff --git a/Engine/Plugins/Experimental/NNE/Source/NNXRuntimeRDG/Private/NNXRuntimeHLSLHelper.cpp b/Engine/Plugins/Experimental/NNE/Source/NNXRuntimeRDG/Private/NNXRuntimeHLSLHelper.cpp
deleted file mode 100644
diff --git a/Engine/Plugins/Experimental/NNE/Source/NNXRuntimeRDG/Private/NNXRuntimeHLSLHelper.h b/Engine/Plugins/Experimental/NNE/Source/NNXRuntimeRDG/Private/NNXRuntimeHLSLHelper.h
deleted file mode 100644
diff --git a/Engine/Plugins/Experimental/NNE/Source/NNXRuntimeRDG/Private/NNXRuntimeHLSLOp.h b/Engine/Plugins/Experimental/NNE/Source/NNXRuntimeRDG/Private/NNXRuntimeHLSLOp.h
deleted file mode 100644
diff --git a/Engine/Plugins/Experimental/NNE/Source/NNXRuntimeRDG/Private/NNXRuntimeRDG.cpp b/Engine/Plugins/Experimental/NNE/Source/NNXRuntimeRDG/Private/NNXRuntimeRDG.cpp
deleted file mode 100644
diff --git a/Engine/Plugins/Experimental/NNE/Source/NNXRuntimeRDG/Private/NNXRuntimeRDG.h b/Engine/Plugins/Experimental/NNE/Source/NNXRuntimeRDG/Private/NNXRuntimeRDG.h
deleted file mode 100644
diff --git a/Engine/Plugins/Experimental/NNE/Source/NNXRuntimeRDG/Private/NNXRuntimeRDGModule.cpp b/Engine/Plugins/Experimental/NNE/Source/NNXRuntimeRDG/Private/NNXRuntimeRDGModule.cpp
deleted file mode 100644
diff --git a/Engine/Plugins/Experimental/NNE/Source/NNXRuntimeRDG/Public/NNERuntimeRDGNeuralPostProcessing.h b/Engine/Plugins/Experimental/NNE/Source/NNXRuntimeRDG/Public/NNERuntimeRDGNeuralPostProcessing.h
deleted file mode 100644
