NNE - Cleanup of NNERuntimeORTCpu

diff --git a/Engine/Plugins/Experimental/NNE/NNE.uplugin b/Engine/Plugins/Experimental/NNE/NNE.uplugin
--- a/Engine/Plugins/Experimental/NNE/NNE.uplugin
+++ b/Engine/Plugins/Experimental/NNE/NNE.uplugin
@@ -50,7 +50,7 @@
 			"PlatformAllowList": [ "Win64", "Linux", "Mac" ]
 		},
 		{
-			"Name": "NNXRuntimeCPU",
+			"Name": "NNERuntimeORTCpu",
 			"Type": "RuntimeAndProgram",
 			"LoadingPhase": "Default",
 			"PlatformAllowList": [ "Win64", "Linux", "Mac" ]
diff --git a/Engine/Plugins/Experimental/NNE/Source/NNEQA/Private/NNEQAUtils.cpp b/Engine/Plugins/Experimental/NNE/Source/NNEQA/Private/NNEQAUtils.cpp
--- a/Engine/Plugins/Experimental/NNE/Source/NNEQA/Private/NNEQAUtils.cpp
+++ b/Engine/Plugins/Experimental/NNE/Source/NNEQA/Private/NNEQAUtils.cpp
@@ -685,7 +685,7 @@ namespace UE::NNEQA::Private
 		UE_LOG(LogNNE, Display, TEXT("Starting tests of '%s'"), *TestSetup.TargetName);
 
 		// Reference runtime
-		const FString& RefName = TEXT("NNERuntimeCPU");
+		const FString& RefName = TEXT("NNERuntimeORTCpu");
 		const float AbsoluteTolerance = TestSetup.GetAbsoluteToleranceForRuntime(RefName);
 		const float RelativeTolerance = TestSetup.GetRelativeToleranceForRuntime(RefName);
 		TArray<TArray<char>> RefOutputMemBuffers;
diff --git a/Engine/Plugins/Experimental/NNE/Source/NNEQA/Resources/NNEQATestDesc.json b/Engine/Plugins/Experimental/NNE/Source/NNEQA/Resources/NNEQATestDesc.json
--- a/Engine/Plugins/Experimental/NNE/Source/NNEQA/Resources/NNEQATestDesc.json
+++ b/Engine/Plugins/Experimental/NNE/Source/NNEQA/Resources/NNEQATestDesc.json
@@ -9,7 +9,7 @@
 					"target": "resnet50-v2-7",
 					"tags": [ "SmokeTest" ],
 					"runtimes": [
-						{ "name": "NNERuntimeORTDmlEditor" },
+						{ "name": "NNERuntimeORTDml" },
 						{ "name": "NNXRuntimeHlsl",  "skip": true, "comment": "Missing BatchNormalisation operator." },
 						{ "name": "NNXRuntimeDml",  "skip": true, "comment": "Model not supported yet." }
 					],
@@ -18,7 +18,7 @@
 				{
 					"target": "Mosaic-9",
 					"runtimes": [
-						{ "name": "NNERuntimeORTDmlEditor" },
+						{ "name": "NNERuntimeORTDml" },
 						{ "name": "NNERuntimeORTCpuEditor" },
 						{ "name": "NNXRuntimeHlsl" },
 						{ "name": "NNXRuntimeDml",  "skip": true, "comment": "Model not supported yet." }
@@ -35,7 +35,7 @@
 					"target": "NeuralMorphModel_global",
 					"tags": [ "SmokeTest" ],
 					"runtimes": [
-						{ "name": "NNERuntimeORTDmlEditor" },
+						{ "name": "NNERuntimeORTDml" },
 						{ "name": "NNXRuntimeHlsl" },
 						{ "name": "NNXRuntimeDml" }
 					],
@@ -45,7 +45,7 @@
 					"target": "NeuralMorphModel_global_Dax",
 					"tags": [ "SmokeTest" ],
 					"runtimes": [
-						{ "name": "NNERuntimeORTDmlEditor", "skip": true, "comment": "" },
+						{ "name": "NNERuntimeORTDml", "skip": true, "comment": "" },
 						{ "name": "NNXRuntimeHlsl",	"skip": true, "comment": "Model not supported yet."	},
 						{ "name": "NNXRuntimeDml",  "skip": true, "comment": "Model not supported yet." }
 					]
@@ -54,7 +54,7 @@
 					"target": "NeuralMorphModel_Local",
 					"tags": [ "SmokeTest" ],
 					"runtimes": [
-						{ "name": "NNERuntimeORTDmlEditor", "absolute_tolerance": 1e-6 },
+						{ "name": "NNERuntimeORTDml", "absolute_tolerance": 1e-6 },
 						{ "name": "NNXRuntimeHlsl",  "skip": true, "comment": "Missing Reshape operator." },
 						{ "name": "NNXRuntimeDml",  "skip": true, "comment": "Model not supported yet." }
 					]
@@ -62,7 +62,7 @@
 				{
 					"target": "VertexDeltaModel",
 					"runtimes": [
-						{ "name": "NNERuntimeORTDmlEditor", "absolute_tolerance": 1e-4 },
+						{ "name": "NNERuntimeORTDml", "absolute_tolerance": 1e-4 },
 						{ "name": "NNXRuntimeHlsl", "absolute_tolerance": 1e-4 },
 						{ "name": "NNXRuntimeDml",  "skip": true, "comment": "Model not supported yet." }
 					]
@@ -122,23 +122,23 @@
 			"targets": [
 				{	"target": "Abs", "additional_datasets": [ "Elementwise.Slow.OneInput" ]},
 				{	"target": "Acos", "tags": [ "SmokeTest" ], "runtimes": [
-						{ "name": "NNERuntimeORTDmlEditor", "absolute_tolerance":  1e-4},
+						{ "name": "NNERuntimeORTDml", "absolute_tolerance":  1e-4},
 						{ "name": "NNXRuntimeDml",    "absolute_tolerance":  1e-4},
 						{ "name": "NNXRuntimeHlsl",   "absolute_tolerance":  2e-4}
 					]},
 				{	"target": "Acosh" },
 				{	"target": "Asin", "runtimes": [
-						{ "name": "NNERuntimeORTDmlEditor", "absolute_tolerance":  1e-4},
+						{ "name": "NNERuntimeORTDml", "absolute_tolerance":  1e-4},
 						{ "name": "NNXRuntimeDml",    "absolute_tolerance":  1e-4},
 						{ "name": "NNXRuntimeHlsl",   "absolute_tolerance":  2e-4}
 					]},
 				{	"target": "Asinh", "runtimes": [
-						{ "name": "NNERuntimeORTDmlEditor", "absolute_tolerance":  1e-4},
+						{ "name": "NNERuntimeORTDml", "absolute_tolerance":  1e-4},
 						{ "name": "NNXRuntimeDml",    "absolute_tolerance":  1e-4},
 						{ "name": "NNXRuntimeHlsl",   "absolute_tolerance":  2e-4}
 					]},
 				{	"target": "Atan", "runtimes": [
-						{ "name": "NNERuntimeORTDmlEditor", "absolute_tolerance":  1e-4},
+						{ "name": "NNERuntimeORTDml", "absolute_tolerance":  1e-4},
 						{ "name": "NNXRuntimeDml",    "absolute_tolerance":  1e-4},
 						{ "name": "NNXRuntimeHlsl",   "absolute_tolerance":  1.5e-3}
 					]},
@@ -153,12 +153,12 @@
 				{	"target": "Cos" },
 				{	"target": "Cosh" },
 				{	"target": "Elu", "comment": "TODO test with not-default attribute", "runtimes": [
-						{ "name": "NNERuntimeORTDmlEditor", "absolute_tolerance":  1e-4},
+						{ "name": "NNERuntimeORTDml", "absolute_tolerance":  1e-4},
 						{ "name": "NNXRuntimeDml",    "absolute_tolerance":  1e-4},
 						{ "name": "NNXRuntimeHlsl",   "absolute_tolerance":  1e-4}
 					]},
 				{	"target": "Erf", "runtimes": [
-						{ "name": "NNERuntimeORTDmlEditor", "absolute_tolerance":  1e-4},
+						{ "name": "NNERuntimeORTDml", "absolute_tolerance":  1e-4},
 						{ "name": "NNXRuntimeDml",    "absolute_tolerance":  1e-4},
 						{ "name": "NNXRuntimeHlsl",   "absolute_tolerance":  5e-4, "comment": "NNXRuntimeHLSL Erf implementation is less precise than the DML one"}
 					]},
@@ -170,7 +170,7 @@
 				{	"target": "HardSwish", "runtimes": [ { "name": "NNXRuntimeDml", "skip": true, "comment": "Not implemented atm on NNXRuntimeDML."} ]},
 				{	"target": "LeakyRelu", "additional_datasets": [ "Elementwise.LeakyRelu.AttributeTest" ] },
 				{	"target": "Log", "runtimes": [
-						{ "name": "NNERuntimeORTDmlEditor"},
+						{ "name": "NNERuntimeORTDml"},
 						{ "name": "NNXRuntimeDml"},
 						{ "name": "NNXRuntimeHlsl"}
 					], "absolute_tolerance": 1e-6 },
@@ -183,14 +183,14 @@
 				{	"target": "Relu"},
 				{	"target": "Round"},
 				{	"target": "Selu", "comment": "TODO test with not-default attribute", "runtimes": [
-						{ "name": "NNERuntimeORTDmlEditor", "absolute_tolerance":  1e-4},
+						{ "name": "NNERuntimeORTDml", "absolute_tolerance":  1e-4},
 						{ "name": "NNXRuntimeDml",    "absolute_tolerance":  1e-4},
 						{ "name": "NNXRuntimeHlsl",   "absolute_tolerance":  1e-4}
 					]},
 				{	"target": "Sigmoid"},
 				{	"target": "Sign" },
 				{	"target": "Sin", "runtimes": [
-						{ "name": "NNERuntimeORTDmlEditor", "absolute_tolerance":  1e-4},
+						{ "name": "NNERuntimeORTDml", "absolute_tolerance":  1e-4},
 						{ "name": "NNXRuntimeDml",    "absolute_tolerance":  1e-4},
 						{ "name": "NNXRuntimeHlsl",   "absolute_tolerance":  1e-4}
 					]},
@@ -199,12 +199,12 @@
 				{	"target": "Softsign"},
 				{	"target": "Sqrt"},
 				{	"target": "Tan", "runtimes": [
-						{ "name": "NNERuntimeORTDmlEditor", "absolute_tolerance":  1e-4},
+						{ "name": "NNERuntimeORTDml", "absolute_tolerance":  1e-4},
 						{ "name": "NNXRuntimeDml",    "absolute_tolerance":  1e-4},
 						{ "name": "NNXRuntimeHlsl",   "absolute_tolerance":  1e-4}
 					]},
 				{	"target": "Tanh", "runtimes": [
-						{ "name": "NNERuntimeORTDmlEditor", "absolute_tolerance":  1e-4},
+						{ "name": "NNERuntimeORTDml", "absolute_tolerance":  1e-4},
 						{ "name": "NNXRuntimeDml",    "absolute_tolerance":  1e-4},
 						{ "name": "NNXRuntimeHlsl",   "absolute_tolerance":  1e-4}
 					]}
@@ -278,7 +278,7 @@
 			"category": "Convolution",
 			"runtimes": [
 				{ "name": "NNXRuntimeDml", "skip": true	},
-				{ "name": "NNERuntimeORTDmlEditor" },
+				{ "name": "NNERuntimeORTDml" },
 				{ "name": "NNXRuntimeHlsl" }
 			],
 			"targets": [
@@ -289,7 +289,7 @@
 			"category": "ConvTranspose",
 			"runtimes": [
 				{ "name": "NNXRuntimeDml", "skip": true	},
-				{ "name": "NNERuntimeORTDmlEditor" },
+				{ "name": "NNERuntimeORTDml" },
 				{ "name": "NNXRuntimeHlsl" }
 			],
 			"targets": [
@@ -475,7 +475,7 @@
 				{ "inputs": [ { "shape": [ 3, 4 ] }, { "shape": [ 4, 5 ] } ], "outputs": [ { "shape": [ 3, 5 ] } ]}
 			],
 			"runtimes": [
-				{ "name": "NNERuntimeORTDmlEditor", "skip": true, "comment": "Variable dimension without bias is failing for NNERuntimeORTDmlEditor" }
+				{ "name": "NNERuntimeORTDml", "skip": true, "comment": "Variable dimension without bias is failing for NNERuntimeORTDml" }
 			]
 		},
 		{
@@ -485,7 +485,7 @@
 				{ "inputs": [ { "shape": [ 4, 3 ] }, { "shape": [ 4, 5 ] } ], "outputs": [ { "shape": [ 3, 5 ] } ]}
 			],
 			"runtimes": [
-				{ "name": "NNERuntimeORTDmlEditor", "skip": true, "comment": "Variable dimension without bias is failing for NNERuntimeORTDmlEditor" }
+				{ "name": "NNERuntimeORTDml", "skip": true, "comment": "Variable dimension without bias is failing for NNERuntimeORTDml" }
 			]
 		},
 		{
@@ -495,7 +495,7 @@
 				{ "inputs": [ { "shape": [ 3, 4 ] }, { "shape": [ 5, 4 ] } ], "outputs": [ { "shape": [ 3, 5 ] } ]}
 			],
 			"runtimes": [
-				{ "name": "NNERuntimeORTDmlEditor", "skip": true, "comment": "Variable dimension without bias is failing for NNERuntimeORTDmlEditor" }
+				{ "name": "NNERuntimeORTDml", "skip": true, "comment": "Variable dimension without bias is failing for NNERuntimeORTDml" }
 			]
 		},
 		{
@@ -505,7 +505,7 @@
 				{ "inputs": [ { "shape": [ 4, 3 ] }, { "shape": [ 5, 4 ] } ], "outputs": [ { "shape": [ 3, 5 ] } ]}
 			],
 			"runtimes": [
-				{ "name": "NNERuntimeORTDmlEditor", "skip": true, "comment": "Variable dimension without bias is failing for NNERuntimeORTDmlEditor" }
+				{ "name": "NNERuntimeORTDml", "skip": true, "comment": "Variable dimension without bias is failing for NNERuntimeORTDml" }
 			]
 		},
 		{
@@ -561,7 +561,7 @@
 			"name": "ConvTranspose.Autopad",
 			"datasets": [ { "inputs": [ { "shape": [ 1, 1, 3, 3 ] }, { "shape": [ 1, 2, 3, 3 ] } ], "outputs": [ { "shape": [ 1, 2, 6, 6 ] } ] } ],
 			"runtimes": [
-				{ "name": "NNERuntimeORTDmlEditor", "skip": true, "comment": "TODO INVESTIGATE Returns 'The parameter is incorrect'" },
+				{ "name": "NNERuntimeORTDml", "skip": true, "comment": "TODO INVESTIGATE Returns 'The parameter is incorrect'" },
 				{ "name": "NNXRuntimeHlsl", "skip": true, "comment": "Fails until ORT 1.13, because they implemented 'auto_pad' differently (they flipped 'SAME_LOWER' and 'SAME_UPPER')" }
 			]
 		},
@@ -583,7 +583,7 @@
 		{
 			"name": "MatMul.Rank5",
 			"runtimes": [	
-				{ "name": "NNERuntimeORTDmlEditor", "skip": true, "comment": "DML metacommand only support 2 to 4 for inputs rank" },
+				{ "name": "NNERuntimeORTDml", "skip": true, "comment": "DML metacommand only support 2 to 4 for inputs rank" },
 				{ "name": "NNXRuntimeHlsl"}
 			],
 			"datasets": [
@@ -652,9 +652,9 @@
 		{
 			"name": "InstanceNormalization",
 			"datasets": [
-				{ "inputs": [ { "shape": [ 1, 1, 3 ] }, { "shape": [ 1 ] }, { "shape": [ 1 ] } ], "outputs": [ { "shape": [ 1, 1, 3 ] } ], "runtimes": [ { "name": "NNERuntimeORTDmlEditor", "skip": true, "comment": "Not supported" } ] },
-				{ "inputs": [ { "shape": [ 1, 2, 3 ] }, { "shape": [ 2 ] }, { "shape": [ 2 ] } ], "outputs": [ { "shape": [ 1, 2, 3 ] } ], "runtimes": [ { "name": "NNERuntimeORTDmlEditor", "skip": true, "comment": "Not supported" } ] },
-				{ "inputs": [ { "shape": [ 2, 2, 3 ] }, { "shape": [ 2 ] }, { "shape": [ 2 ] } ], "outputs": [ { "shape": [ 2, 2, 3 ] } ], "runtimes": [ { "name": "NNERuntimeORTDmlEditor", "skip": true, "comment": "Not supported" } ]},
+				{ "inputs": [ { "shape": [ 1, 1, 3 ] }, { "shape": [ 1 ] }, { "shape": [ 1 ] } ], "outputs": [ { "shape": [ 1, 1, 3 ] } ], "runtimes": [ { "name": "NNERuntimeORTDml", "skip": true, "comment": "Not supported" } ] },
+				{ "inputs": [ { "shape": [ 1, 2, 3 ] }, { "shape": [ 2 ] }, { "shape": [ 2 ] } ], "outputs": [ { "shape": [ 1, 2, 3 ] } ], "runtimes": [ { "name": "NNERuntimeORTDml", "skip": true, "comment": "Not supported" } ] },
+				{ "inputs": [ { "shape": [ 2, 2, 3 ] }, { "shape": [ 2 ] }, { "shape": [ 2 ] } ], "outputs": [ { "shape": [ 2, 2, 3 ] } ], "runtimes": [ { "name": "NNERuntimeORTDml", "skip": true, "comment": "Not supported" } ]},
 				{ "inputs": [ { "shape": [ 1, 2, 1, 3 ] }, { "shape": [ 2 ] }, { "shape": [ 2 ] } ], "outputs": [ { "shape": [ 1, 2, 1, 3 ] } ] },
 				{ "inputs": [ { "shape": [ 2, 2, 1, 3 ] }, { "shape": [ 2 ] }, { "shape": [ 2 ] } ], "outputs": [ { "shape": [ 2, 2, 1, 3 ] } ] },
 				{ "inputs": [ { "shape": [ 2, 2, 2, 2 ] }, { "shape": [ 2 ] }, { "shape": [ 2 ] } ], "outputs": [ { "shape": [ 2, 2, 2, 2 ] } ] },
@@ -663,7 +663,7 @@
 					"runtimes": [
 						{ "name": "NNXRuntimeHLSL" },
 						{ "name": "NNERuntimeORTCpuEditor" },
-						{ "name": "NNERuntimeORTDmlEditor", "skip": true, "comment": "Not supported" }
+						{ "name": "NNERuntimeORTDml", "skip": true, "comment": "Not supported" }
 					]
 				},
 				{
@@ -671,7 +671,7 @@
 					"runtimes": [
 						{ "name": "NNXRuntimeHLSL" },
 						{ "name": "NNERuntimeORTCpuEditor" },
-						{ "name": "NNERuntimeORTDmlEditor" }
+						{ "name": "NNERuntimeORTDml" }
 					]
 				},
 				{
@@ -679,7 +679,7 @@
 					"runtimes": [
 						{ "name": "NNXRuntimeHLSL" },
 						{ "name": "NNERuntimeORTCpuEditor" },
-						{ "name": "NNERuntimeORTDmlEditor" }
+						{ "name": "NNERuntimeORTDml" }
 					]
 				},
 				{
@@ -687,7 +687,7 @@
 					"runtimes": [
 						{ "name": "NNXRuntimeHLSL" },
 						{ "name": "NNERuntimeORTCpuEditor" },
-						{ "name": "NNERuntimeORTDmlEditor" }
+						{ "name": "NNERuntimeORTDml" }
 					]
 				},
 				{
@@ -695,7 +695,7 @@
 					"runtimes": [
 						{ "name": "NNXRuntimeHLSL" },
 						{ "name": "NNERuntimeORTCpuEditor" },
-						{ "name": "NNERuntimeORTDmlEditor" }
+						{ "name": "NNERuntimeORTDml" }
 					]
 				},
 				{
@@ -703,7 +703,7 @@
 					"runtimes": [
 						{ "name": "NNXRuntimeHLSL" },
 						{ "name": "NNERuntimeORTCpuEditor" },
-						{ "name": "NNERuntimeORTDmlEditor" }
+						{ "name": "NNERuntimeORTDml" }
 					]
 				},
 				{
@@ -711,7 +711,7 @@
 					"runtimes": [
 						{ "name": "NNXRuntimeHLSL" },
 						{ "name": "NNERuntimeORTCpuEditor" },
-						{ "name": "NNERuntimeORTDmlEditor" }
+						{ "name": "NNERuntimeORTDml" }
 					]
 				},
 				{
@@ -719,7 +719,7 @@
 					"runtimes": [
 						{ "name": "NNXRuntimeHLSL" },
 						{ "name": "NNERuntimeORTCpuEditor" },
-						{ "name": "NNERuntimeORTDmlEditor" }
+						{ "name": "NNERuntimeORTDml" }
 					],
 					"comment": "from 'mosaic-9-sim.onnx'"
 				}
diff --git a/Engine/Plugins/Experimental/NNE/Source/NNERuntimeORT/Private/NNERuntimeORT.cpp b/Engine/Plugins/Experimental/NNE/Source/NNERuntimeORT/Private/NNERuntimeORT.cpp
--- a/Engine/Plugins/Experimental/NNE/Source/NNERuntimeORT/Private/NNERuntimeORT.cpp
+++ b/Engine/Plugins/Experimental/NNE/Source/NNERuntimeORT/Private/NNERuntimeORT.cpp
@@ -14,15 +14,15 @@
 
 using namespace NNX;
 
-FGuid UNNERuntimeORTCpuImpl::GUID = FGuid((int32)'O', (int32)'N', (int32)'N', (int32)'X');
-int32 UNNERuntimeORTCpuImpl::Version = 0x00000001;
+FGuid UNNERuntimeORTDmlImpl::GUID = FGuid((int32)'O', (int32)'D', (int32)'M', (int32)'L');
+int32 UNNERuntimeORTDmlImpl::Version = 0x00000001;
 
-bool UNNERuntimeORTCpuImpl::CanCreateModelData(FString FileType, TConstArrayView<uint8> FileData) const
+bool UNNERuntimeORTDmlImpl::CanCreateModelData(FString FileType, TConstArrayView<uint8> FileData) const
 {
 	return FileType.Compare("onnx", ESearchCase::IgnoreCase) == 0;
 }
 
-TArray<uint8> UNNERuntimeORTCpuImpl::CreateModelData(FString FileType, TConstArrayView<uint8> FileData)
+TArray<uint8> UNNERuntimeORTDmlImpl::CreateModelData(FString FileType, TConstArrayView<uint8> FileData)
 {
 	if (!CanCreateModelData(FileType, FileData))
 	{
@@ -41,70 +41,29 @@ TArray<uint8> UNNERuntimeORTCpuImpl::CreateModelData(FString FileType, TConstArr
 		return {};
 	}
 
-	int32 GuidSize = sizeof(UNNERuntimeORTCpuImpl::GUID);
-	int32 VersionSize = sizeof(UNNERuntimeORTCpuImpl::Version);
+	int32 GuidSize = sizeof(UNNERuntimeORTDmlImpl::GUID);
+	int32 VersionSize = sizeof(UNNERuntimeORTDmlImpl::Version);
 	TArray<uint8> Result;
 	FMemoryWriter Writer(Result);
-	Writer << UNNERuntimeORTCpuImpl::GUID;
-	Writer << UNNERuntimeORTCpuImpl::Version;
+	Writer << UNNERuntimeORTDmlImpl::GUID;
+	Writer << UNNERuntimeORTDmlImpl::Version;
 	Writer.Serialize(OutputModel.Data.GetData(), OutputModel.Data.Num());
 	return Result;
 }
 
-bool UNNERuntimeORTCpuImpl::CanCreateModelCPU(TObjectPtr<UNNEModelData> ModelData) const
-{
-	check(ModelData != nullptr);
-
-	int32 GuidSize = sizeof(UNNERuntimeORTCpuImpl::GUID);
-	int32 VersionSize = sizeof(UNNERuntimeORTCpuImpl::Version);
-	TConstArrayView<uint8> Data = ModelData->GetModelData(GetRuntimeName());
-
-	if (Data.Num() <= GuidSize + VersionSize)
-	{
-		return false;
-	}
-
-	bool bResult = FGenericPlatformMemory::Memcmp(&(Data[0]), &(UNNERuntimeORTCpuImpl::GUID), GuidSize) == 0;
-	bResult &= FGenericPlatformMemory::Memcmp(&(Data[GuidSize]), &(UNNERuntimeORTCpuImpl::Version), VersionSize) == 0;
-	return bResult;
-}
-
-void UNNERuntimeORTCpuImpl::Init()
+void UNNERuntimeORTDmlImpl::Init()
 {
 	check(!ORTEnvironment.IsValid());
 	ORTEnvironment = MakeUnique<Ort::Env>();
 }
 
-TUniquePtr<UE::NNECore::IModelCPU> UNNERuntimeORTCpuImpl::CreateModelCPU(TObjectPtr<UNNEModelData> ModelData)
-{
-	check(ModelData != nullptr);
-	check(ORTEnvironment.IsValid());
-
-	if (!CanCreateModelCPU(ModelData))
-	{
-		return TUniquePtr<UE::NNECore::IModelCPU>();
-	}
-
-	const UE::NNERuntimeORT::Private::FRuntimeConf InConf;
-	UE::NNERuntimeORT::Private::FModelORTCpu* Model = new UE::NNERuntimeORT::Private::FModelORTCpu(ORTEnvironment.Get(), InConf);
-	TConstArrayView<uint8> Data = ModelData->GetModelData(GetRuntimeName());
-
-	if (!Model->Init(Data))
-	{
-		delete Model;
-		return TUniquePtr<UE::NNECore::IModelCPU>();
-	}
-	UE::NNECore::IModelCPU* IModel = static_cast<UE::NNECore::IModelCPU*>(Model);
-	return TUniquePtr<UE::NNECore::IModelCPU>(IModel);
-}
-
 #if PLATFORM_WINDOWS
 bool UNNERuntimeORTDmlImpl::CanCreateModelCPU(TObjectPtr<UNNEModelData> ModelData) const
 {
 	check(ModelData != nullptr);
 
-	int32 GuidSize = sizeof(UNNERuntimeORTCpuImpl::GUID);
-	int32 VersionSize = sizeof(UNNERuntimeORTCpuImpl::Version);
+	int32 GuidSize = sizeof(UNNERuntimeORTDmlImpl::GUID);
+	int32 VersionSize = sizeof(UNNERuntimeORTDmlImpl::Version);
 	TConstArrayView<uint8> Data = ModelData->GetModelData(GetRuntimeName());
 
 	if (Data.Num() <= GuidSize + VersionSize)
@@ -112,8 +71,8 @@ bool UNNERuntimeORTDmlImpl::CanCreateModelCPU(TObjectPtr<UNNEModelData> ModelDat
 		return false;
 	}
 
-	bool bResult = FGenericPlatformMemory::Memcmp(&(Data[0]), &(UNNERuntimeORTCpuImpl::GUID), GuidSize) == 0;
-	bResult &= FGenericPlatformMemory::Memcmp(&(Data[GuidSize]), &(UNNERuntimeORTCpuImpl::Version), VersionSize) == 0;
+	bool bResult = FGenericPlatformMemory::Memcmp(&(Data[0]), &(UNNERuntimeORTDmlImpl::GUID), GuidSize) == 0;
+	bResult &= FGenericPlatformMemory::Memcmp(&(Data[GuidSize]), &(UNNERuntimeORTDmlImpl::Version), VersionSize) == 0;
 	return bResult;
 }
 
diff --git a/Engine/Plugins/Experimental/NNE/Source/NNERuntimeORT/Private/NNERuntimeORT.h b/Engine/Plugins/Experimental/NNE/Source/NNERuntimeORT/Private/NNERuntimeORT.h
--- a/Engine/Plugins/Experimental/NNE/Source/NNERuntimeORT/Private/NNERuntimeORT.h
+++ b/Engine/Plugins/Experimental/NNE/Source/NNERuntimeORT/Private/NNERuntimeORT.h
@@ -20,7 +20,7 @@ NNX_THIRD_PARTY_INCLUDES_END
 #include "NNERuntimeORT.generated.h"
 
 UCLASS()
-class UNNERuntimeORTCpuImpl : public UObject, public INNERuntime, public INNERuntimeCPU
+class UNNERuntimeORTDmlImpl : public UObject, public INNERuntime, public INNERuntimeCPU
 {
 	GENERATED_BODY()
 
@@ -29,27 +29,17 @@ public:
 	static int32 Version;
 
 	TUniquePtr<Ort::Env> ORTEnvironment;
-	UNNERuntimeORTCpuImpl() {};
-	virtual ~UNNERuntimeORTCpuImpl() {}
+	UNNERuntimeORTDmlImpl() {};
+	virtual ~UNNERuntimeORTDmlImpl() {}
 
 	void Init();
 
-	virtual FString GetRuntimeName() const override { return TEXT("NNERuntimeORTCpuEditor"); };
+	virtual FString GetRuntimeName() const override { return TEXT("NNERuntimeORTDml"); };
 	virtual bool IsPlatformSupported(const ITargetPlatform* TargetPlatform) const override { return true; };
 
 	virtual bool CanCreateModelData(FString FileType, TConstArrayView<uint8> FileData) const override;
 	virtual TArray<uint8> CreateModelData(FString FileType, TConstArrayView<uint8> FileData) override;
 
-	virtual bool CanCreateModelCPU(TObjectPtr<UNNEModelData> ModelData) const override;
-	virtual TUniquePtr<UE::NNECore::IModelCPU> CreateModelCPU(TObjectPtr<UNNEModelData> ModelData) override;
-};
-
-UCLASS()
-class UNNERuntimeORTDmlImpl : public UNNERuntimeORTCpuImpl
-{
-	GENERATED_BODY()
-
-	virtual FString GetRuntimeName() const override { return TEXT("NNERuntimeORTDmlEditor"); };
 	virtual bool CanCreateModelCPU(TObjectPtr<UNNEModelData> ModelData) const override;
 	virtual TUniquePtr<UE::NNECore::IModelCPU> CreateModelCPU(TObjectPtr<UNNEModelData> ModelData) override;
 };
\ No newline at end of file
diff --git a/Engine/Plugins/Experimental/NNE/Source/NNERuntimeORT/Private/NNERuntimeORTModel.cpp b/Engine/Plugins/Experimental/NNE/Source/NNERuntimeORT/Private/NNERuntimeORTModel.cpp
--- a/Engine/Plugins/Experimental/NNE/Source/NNERuntimeORT/Private/NNERuntimeORTModel.cpp
+++ b/Engine/Plugins/Experimental/NNE/Source/NNERuntimeORT/Private/NNERuntimeORTModel.cpp
@@ -57,8 +57,8 @@ namespace UE::NNERuntimeORT::Private
 		DECLARE_SCOPE_CYCLE_COUNTER(TEXT("FModelORT_Init"), STAT_FModelORT_Init, STATGROUP_NNE);
 		
 		// Get the header size
-		int32 GuidSize = sizeof(UNNERuntimeORTCpuImpl::GUID);
-		int32 VersionSize = sizeof(UNNERuntimeORTCpuImpl::Version);
+		int32 GuidSize = sizeof(UNNERuntimeORTDmlImpl::GUID);
+		int32 VersionSize = sizeof(UNNERuntimeORTDmlImpl::Version);
 
 		// Clean previous networks
 		bIsLoaded = false;
diff --git a/Engine/Plugins/Experimental/NNE/Source/NNERuntimeORT/Private/NNERuntimeORTModule.cpp b/Engine/Plugins/Experimental/NNE/Source/NNERuntimeORT/Private/NNERuntimeORTModule.cpp
--- a/Engine/Plugins/Experimental/NNE/Source/NNERuntimeORT/Private/NNERuntimeORTModule.cpp
+++ b/Engine/Plugins/Experimental/NNE/Source/NNERuntimeORT/Private/NNERuntimeORTModule.cpp
@@ -58,17 +58,6 @@ void FNNERuntimeORTModule::StartupModule()
 
 	Ort::InitApi();
 	
-	// NNE runtime ORT Cpu startup
-	NNERuntimeORTCpu = NewObject<UNNERuntimeORTCpuImpl>();
-	if (NNERuntimeORTCpu.IsValid())
-	{
-		TWeakInterfacePtr<INNERuntime> RuntimeCPUInterface(NNERuntimeORTCpu.Get());
-
-		NNERuntimeORTCpu->Init();
-		NNERuntimeORTCpu->AddToRoot();
-		UE::NNECore::RegisterRuntime(RuntimeCPUInterface);
-	}
-
 #if PLATFORM_WINDOWS
 	// NNE runtime ORT Dml startup
 	NNERuntimeORTDml = NewObject<UNNERuntimeORTDmlImpl>();
@@ -85,16 +74,6 @@ void FNNERuntimeORTModule::StartupModule()
 
 void FNNERuntimeORTModule::ShutdownModule()
 {
-	// NNE runtime ORT Cpu shutdown
-	if (NNERuntimeORTCpu.IsValid())
-	{
-		TWeakInterfacePtr<INNERuntime> RuntimeCPUInterface(NNERuntimeORTCpu.Get());
-
-		UE::NNECore::UnregisterRuntime(RuntimeCPUInterface);
-		NNERuntimeORTCpu->RemoveFromRoot();
-		NNERuntimeORTCpu = TWeakObjectPtr<UNNERuntimeORTCpuImpl>(nullptr);
-	}
-
 #if PLATFORM_WINDOWS
 	// NNE runtime ORT Dml shutdown
 	if (NNERuntimeORTDml.IsValid())
diff --git a/Engine/Plugins/Experimental/NNE/Source/NNERuntimeORT/Private/NNERuntimeORTModule.h b/Engine/Plugins/Experimental/NNE/Source/NNERuntimeORT/Private/NNERuntimeORTModule.h
--- a/Engine/Plugins/Experimental/NNE/Source/NNERuntimeORT/Private/NNERuntimeORTModule.h
+++ b/Engine/Plugins/Experimental/NNE/Source/NNERuntimeORT/Private/NNERuntimeORTModule.h
@@ -9,13 +9,11 @@
 #include "UObject/WeakObjectPtrTemplates.h"
 
 class UNNERuntimeORTDmlImpl;
-class UNNERuntimeORTCpuImpl;
 
 class FNNERuntimeORTModule : public IModuleInterface
 {
 
 public:
-	TWeakObjectPtr<UNNERuntimeORTCpuImpl> NNERuntimeORTCpu{ nullptr };
 	TWeakObjectPtr<UNNERuntimeORTDmlImpl> NNERuntimeORTDml{ nullptr };
 
 	// Begin IModuleInterface
diff --git a/Engine/Plugins/Experimental/NNE/Source/NNERuntimeORTCpu/NNERuntimeORTCpu.Build.cs b/Engine/Plugins/Experimental/NNE/Source/NNERuntimeORTCpu/NNERuntimeORTCpu.Build.cs
new file mode 100644
--- /dev/null
+++ b/Engine/Plugins/Experimental/NNE/Source/NNERuntimeORTCpu/NNERuntimeORTCpu.Build.cs
@@ -0,0 +1,35 @@
+// Copyright Epic Games, Inc. All Rights Reserved.
+
+using UnrealBuildTool;
+
+public class NNERuntimeORTCpu : ModuleRules
+{
+	public NNERuntimeORTCpu( ReadOnlyTargetRules Target ) : base( Target )
+	{
+		bool bIsORTSupported = (Target.Platform == UnrealTargetPlatform.Win64 || Target.Platform == UnrealTargetPlatform.Linux || Target.Platform == UnrealTargetPlatform.Mac);
+		
+		PCHUsage = PCHUsageMode.UseExplicitOrSharedPCHs;
+
+		PublicDependencyModuleNames.AddRange
+		(
+			new string[] 
+			{
+				"Core",
+				"Engine",
+				"NNXCore",
+				"NNEProfiling",
+				"ORTHelper",
+			}
+		);
+
+		PrivateDependencyModuleNames.AddRange
+		(
+			new string[] 
+			{
+				"CoreUObject",
+				"NNEUtils",
+				"NNEOnnxruntime"
+			}
+		);
+	}
+}
diff --git a/Engine/Plugins/Experimental/NNE/Source/NNERuntimeORTCpu/Private/NNERuntimeORTCpu.cpp b/Engine/Plugins/Experimental/NNE/Source/NNERuntimeORTCpu/Private/NNERuntimeORTCpu.cpp
new file mode 100644
--- /dev/null
+++ b/Engine/Plugins/Experimental/NNE/Source/NNERuntimeORTCpu/Private/NNERuntimeORTCpu.cpp
@@ -0,0 +1,89 @@
+// Copyright Epic Games, Inc. All Rights Reserved.
+
+#include "NNERuntimeORTCpu.h"
+#include "NNERuntimeORTCpuModel.h"
+#include "NNERuntimeORTCpuUtils.h"
+#include "NNEUtilsModelOptimizer.h"
+#include "NNECoreAttributeMap.h"
+#include "NNECoreModelData.h"
+#include "NNEProfilingTimer.h"
+#include "RedirectCoutAndCerrToUeLog.h"
+
+#include "NNXModelOptimizerInterface.h"
+
+FGuid UNNERuntimeORTCpuImpl::GUID = FGuid((int32)'O', (int32)'C', (int32)'P', (int32)'U');
+int32 UNNERuntimeORTCpuImpl::Version = 0x00000001;
+
+bool UNNERuntimeORTCpuImpl::CanCreateModelData(FString FileType, TConstArrayView<uint8> FileData) const
+{
+	return FileType.Compare("onnx", ESearchCase::IgnoreCase) == 0;
+}
+
+TArray<uint8> UNNERuntimeORTCpuImpl::CreateModelData(FString FileType, TConstArrayView<uint8> FileData)
+{
+	if (!CanCreateModelData(FileType, FileData))
+	{
+		return {};
+	}
+
+	TUniquePtr<NNX::IModelOptimizer> Optimizer = UE::NNEUtils::Internal::CreateONNXToONNXModelOptimizer();
+
+	FNNIModelRaw InputModel;
+	InputModel.Data = FileData;
+	InputModel.Format = ENNXInferenceFormat::ONNX;
+	FNNIModelRaw OutputModel;
+	UE::NNEUtils::Internal::FOptimizerOptionsMap Options;
+	if (!Optimizer->Optimize(InputModel, OutputModel, Options))
+	{
+		return {};
+	}
+
+	int32 GuidSize = sizeof(UNNERuntimeORTCpuImpl::GUID);
+	int32 VersionSize = sizeof(UNNERuntimeORTCpuImpl::Version);
+	TArray<uint8> Result;
+	FMemoryWriter Writer(Result);
+	Writer << UNNERuntimeORTCpuImpl::GUID;
+	Writer << UNNERuntimeORTCpuImpl::Version;
+	Writer.Serialize(OutputModel.Data.GetData(), OutputModel.Data.Num());
+	return Result;
+}
+
+bool UNNERuntimeORTCpuImpl::CanCreateModelCPU(TObjectPtr<UNNEModelData> ModelData) const
+{
+	check(ModelData != nullptr);
+	
+	int32 GuidSize = sizeof(UNNERuntimeORTCpuImpl::GUID);
+	int32 VersionSize = sizeof(UNNERuntimeORTCpuImpl::Version);
+	TConstArrayView<uint8> Data = ModelData->GetModelData(GetRuntimeName());
+	
+	if (Data.Num() <= GuidSize + VersionSize)
+	{
+		return false;
+	}
+	
+	bool bResult = FGenericPlatformMemory::Memcmp(&(Data[0]), &(UNNERuntimeORTCpuImpl::GUID), GuidSize) == 0;
+	bResult &= FGenericPlatformMemory::Memcmp(&(Data[GuidSize]), &(UNNERuntimeORTCpuImpl::Version), VersionSize) == 0;
+	return bResult;
+}
+
+TUniquePtr<UE::NNECore::IModelCPU> UNNERuntimeORTCpuImpl::CreateModelCPU(TObjectPtr<UNNEModelData> ModelData)
+{
+	check(ModelData != nullptr);
+	
+	if (!CanCreateModelCPU(ModelData))
+	{
+		return TUniquePtr<UE::NNECore::IModelCPU>();
+	}
+
+	const UE::NNERuntimeORTCpu::Private::FRuntimeConf InConf;
+	UE::NNERuntimeORTCpu::Private::FModelCPU* Model = new UE::NNERuntimeORTCpu::Private::FModelCPU(&NNEEnvironmentCPU, InConf);
+	TConstArrayView<uint8> Data = ModelData->GetModelData(GetRuntimeName());
+	
+	if (!Model->Init(Data))
+	{
+		delete Model;
+		return TUniquePtr<UE::NNECore::IModelCPU>();
+	}
+	UE::NNECore::IModelCPU* IModel = static_cast<UE::NNECore::IModelCPU*>(Model);
+	return TUniquePtr<UE::NNECore::IModelCPU>(IModel);
+}
\ No newline at end of file
diff --git a/Engine/Plugins/Experimental/NNE/Source/NNERuntimeORTCpu/Private/NNERuntimeORTCpu.h b/Engine/Plugins/Experimental/NNE/Source/NNERuntimeORTCpu/Private/NNERuntimeORTCpu.h
new file mode 100644
--- /dev/null
+++ b/Engine/Plugins/Experimental/NNE/Source/NNERuntimeORTCpu/Private/NNERuntimeORTCpu.h
@@ -0,0 +1,42 @@
+// Copyright Epic Games, Inc. All Rights Reserved.
+
+#pragma once
+
+#include "CoreMinimal.h"
+#include "NNECoreRuntime.h"
+#include "NNECoreRuntimeCPU.h"
+#include "UObject/Class.h"
+#include "UObject/Object.h"
+#include "UObject/UObjectBaseUtility.h"
+
+#include "NNXThirdPartyWarningDisabler.h"
+NNX_THIRD_PARTY_INCLUDES_START
+#undef check
+#undef TEXT
+#include "core/session/onnxruntime_cxx_api.h"
+NNX_THIRD_PARTY_INCLUDES_END
+
+#include "NNERuntimeORTCpu.generated.h"
+
+UCLASS()
+class UNNERuntimeORTCpuImpl : public UObject, public INNERuntime, public INNERuntimeCPU
+{
+	GENERATED_BODY()
+
+public:
+	static FGuid GUID;
+	static int32 Version;
+
+	Ort::Env NNEEnvironmentCPU;
+	UNNERuntimeORTCpuImpl() {};
+	virtual ~UNNERuntimeORTCpuImpl() {}
+		
+	virtual FString GetRuntimeName() const override { return TEXT("NNERuntimeORTCpu"); };
+	virtual bool IsPlatformSupported(const ITargetPlatform* TargetPlatform) const override { return true; };
+
+	virtual bool CanCreateModelData(FString FileType, TConstArrayView<uint8> FileData) const override;
+	virtual TArray<uint8> CreateModelData(FString FileType, TConstArrayView<uint8> FileData) override;
+
+	virtual bool CanCreateModelCPU(TObjectPtr<UNNEModelData> ModelData) const override;
+	virtual TUniquePtr<UE::NNECore::IModelCPU> CreateModelCPU(TObjectPtr<UNNEModelData> ModelData) override;
+};
\ No newline at end of file
diff --git a/Engine/Plugins/Experimental/NNE/Source/NNERuntimeORTCpu/Private/NNERuntimeORTCpuModel.cpp b/Engine/Plugins/Experimental/NNE/Source/NNERuntimeORTCpu/Private/NNERuntimeORTCpuModel.cpp
new file mode 100644
--- /dev/null
+++ b/Engine/Plugins/Experimental/NNE/Source/NNERuntimeORTCpu/Private/NNERuntimeORTCpuModel.cpp
@@ -0,0 +1,301 @@
+// Copyright Epic Games, Inc. All Rights Reserved.
+#include "NNERuntimeORTCpuModel.h"
+#include "NNERuntimeORTCpu.h"
+#include "NNERuntimeORTCpuUtils.h"
+#include "NNEProfilingTimer.h"
+#include "RedirectCoutAndCerrToUeLog.h"
+
+namespace UE::NNERuntimeORTCpu::Private
+{
+
+	FModelCPU::FModelCPU() :
+		bIsLoaded(false),
+		bHasRun(false)
+	{ }
+
+	FModelCPU::FModelCPU(
+		Ort::Env* InORTEnvironment,
+		const FRuntimeConf& InRuntimeConf) :
+		bIsLoaded(false),
+		bHasRun(false),
+		RuntimeConf(InRuntimeConf),
+		ORTEnvironment(InORTEnvironment)
+	{ }
+
+	bool FModelCPU::Init(TConstArrayView<uint8> ModelData)
+	{
+		DECLARE_SCOPE_CYCLE_COUNTER(TEXT("FModelCPU_Init"), STAT_FModelCPU_Init, STATGROUP_NNE);
+		
+		// Get the header size
+		int32 GuidSize = sizeof(UNNERuntimeORTCpuImpl::GUID);
+		int32 VersionSize = sizeof(UNNERuntimeORTCpuImpl::Version);
+
+		// Clean previous networks
+		bIsLoaded = false;
+		TConstArrayView<uint8> ModelBuffer = TConstArrayView<uint8>(&(ModelData.GetData()[GuidSize + VersionSize]), ModelData.Num() - GuidSize - VersionSize);
+
+		// Checking Inference Model 
+		{
+			if (ModelBuffer.Num() == 0)
+			{
+				UE_LOG(LogNNE, Warning, TEXT("FModelCPU::Load(): Input model data is empty."));
+				return false;
+			}
+
+		}
+
+#if WITH_EDITOR
+		try
+#endif //WITH_EDITOR
+		{
+			const FRedirectCoutAndCerrToUeLog RedirectCoutAndCerrToUeLog;
+			if (!InitializedAndConfigureMembers())
+			{
+				UE_LOG(LogNNE, Warning, TEXT("Load(): InitializedAndConfigureMembers failed."));
+				return false;
+			}
+
+			{
+				DECLARE_SCOPE_CYCLE_COUNTER(TEXT("FModelCPU_Init_CreateORTSession"), STAT_FModelCPU_Init_CreateORTSession, STATGROUP_NNE);
+				// Read model from InferenceModel
+				Session = MakeUnique<Ort::Session>(*ORTEnvironment, ModelBuffer.GetData(), ModelBuffer.Num(), *SessionOptions);
+			}
+
+			if (!ConfigureTensors(true))
+			{
+				UE_LOG(LogNNE, Warning, TEXT("Load(): Failed to configure Inputs tensors."));
+				return false;
+			}
+
+			if (!ConfigureTensors(false))
+			{
+				UE_LOG(LogNNE, Warning, TEXT("Load(): Failed to configure Outputs tensors."));
+				return false;
+			}
+		}
+#if WITH_EDITOR
+		catch (const std::exception& Exception)
+		{
+			UE_LOG(LogNNE, Error, TEXT("%s"), UTF8_TO_TCHAR(Exception.what()));
+			return false;
+		}
+#endif //WITH_EDITOR
+
+		bIsLoaded = true;
+
+		// Reset Stats
+		ResetStats();
+
+		return IsLoaded();
+	}
+
+	bool FModelCPU::IsLoaded() const
+	{
+		return bIsLoaded;
+	}
+
+	bool FModelCPU::InitializedAndConfigureMembers()
+	{
+		// Initialize 
+		// Set up ORT and create an environment
+		Allocator = MakeUnique<Ort::AllocatorWithDefaultOptions>();
+		AllocatorInfo = MakeUnique<Ort::MemoryInfo>(Ort::MemoryInfo::CreateCpu(OrtDeviceAllocator, OrtMemTypeCPU));
+
+		// Configure Session
+		SessionOptions = MakeUnique<Ort::SessionOptions>();
+
+		// Configure number threads
+		//SessionOptions->SetIntraOpNumThreads(RuntimeConf.NumberOfThreads);
+		// Uncomment if you want to change the priority of the threads, by default is TPri_Normal
+		//SessionOptions->SetPriorityOpThreads(RuntimeConf.ThreadPriority);
+
+		// ORT CPU
+		SessionOptions->SetGraphOptimizationLevel(RuntimeConf.OptimizationLevel); // ORT_ENABLE_ALL, ORT_ENABLE_EXTENDED, ORT_ENABLE_BASIC, ORT_DISABLE_ALL
+		SessionOptions->EnableCpuMemArena();
+
+		return true;
+	}
+
+
+	bool FModelCPU::ConfigureTensors(const bool InIsInput)
+	{
+		const bool bIsInput = InIsInput;
+
+		const uint32 NumberTensors = bIsInput ? Session->GetInputCount() : Session->GetOutputCount();
+
+		TArray<NNECore::FTensorDesc>& SymbolicTensorDescs = bIsInput ? InputSymbolicTensors : OutputSymbolicTensors;
+		TArray<ONNXTensorElementDataType>& TensorsORTType = bIsInput ? InputTensorsORTType : OutputTensorsORTType;
+		TArray<const char*>& TensorNames = bIsInput ? InputTensorNames : OutputTensorNames;
+
+		for (uint32 TensorIndex = 0; TensorIndex < NumberTensors; ++TensorIndex)
+		{
+			// Get Tensor name
+			const char* CurTensorName = bIsInput ? Session->GetInputName(TensorIndex, *Allocator) : Session->GetOutputName(TensorIndex, *Allocator);
+			TensorNames.Emplace(CurTensorName);
+
+			// Get node type
+			Ort::TypeInfo CurrentTypeInfo = bIsInput ? Session->GetInputTypeInfo(TensorIndex) : Session->GetOutputTypeInfo(TensorIndex);
+			Ort::TensorTypeAndShapeInfo CurrentTensorInfo = CurrentTypeInfo.GetTensorTypeAndShapeInfo();
+			const ONNXTensorElementDataType ONNXTensorElementDataTypeEnum = CurrentTensorInfo.GetElementType();
+			CurrentTypeInfo.release();
+
+			TensorsORTType.Emplace(ONNXTensorElementDataTypeEnum);
+
+			std::pair<ENNETensorDataType, uint64> TypeAndSize = TranslateTensorTypeORTToNNE(ONNXTensorElementDataTypeEnum);
+
+			TArray<int32> ShapeData;
+			ShapeData.Reserve(CurrentTensorInfo.GetShape().size());
+			for (const int64_t CurrentTensorSize : CurrentTensorInfo.GetShape())
+			{
+				ShapeData.Add((int32)CurrentTensorSize);
+			}
+
+			NNECore::FSymbolicTensorShape Shape = NNECore::FSymbolicTensorShape::Make(ShapeData);
+			NNECore::FTensorDesc SymbolicTensorDesc = NNECore::FTensorDesc::Make(FString(CurTensorName), Shape, TypeAndSize.first);
+
+			check(SymbolicTensorDesc.GetElemByteSize() == TypeAndSize.second);
+			SymbolicTensorDescs.Emplace(SymbolicTensorDesc);
+		}
+
+		return true;
+	}
+	
+	int FModelCPU::SetInputTensorShapes(TConstArrayView<NNECore::FTensorShape> InInputShapes)
+	{
+		InputTensors.Empty();
+		OutputTensors.Empty();
+		OutputTensorShapes.Empty();
+
+		// Verify input shape are valid for the model and set InputTensorShapes
+		if (FModelBase<IModelCPU>::SetInputTensorShapes(InInputShapes) != 0)
+		{
+			return -1;
+		}
+
+		// Setup concrete input tensor
+		for (int i = 0; i < InputSymbolicTensors.Num(); ++i)
+		{
+			NNECore::Internal::FTensor Tensor = NNECore::Internal::FTensor::Make(InputSymbolicTensors[i].GetName(), InInputShapes[i], InputSymbolicTensors[i].GetDataType());
+			InputTensors.Emplace(Tensor);
+		}
+
+		// Here model optimization could be done now that we know the input shapes, for some models
+		// this would allow to resolve output shapes here rather than during inference.
+
+		// Setup concrete output shapes only if all model output shapes are concretes, otherwise it will be set during Run()
+		for (NNECore::FTensorDesc SymbolicTensorDesc : OutputSymbolicTensors)
+		{
+			if (SymbolicTensorDesc.GetShape().IsConcrete())
+			{
+				NNECore::Internal::FTensor Tensor = NNECore::Internal::FTensor::MakeFromSymbolicDesc(SymbolicTensorDesc);
+				OutputTensors.Emplace(Tensor);
+				OutputTensorShapes.Emplace(Tensor.GetShape());
+			}
+		}
+		if (OutputTensors.Num() != OutputSymbolicTensors.Num())
+		{
+			OutputTensors.Empty();
+			OutputTensorShapes.Empty();
+		}
+
+		return 0;
+	}
+
+	int FModelCPU::RunSync(TConstArrayView<NNECore::FTensorBindingCPU> InInputBindings, TConstArrayView<NNECore::FTensorBindingCPU> InOutputBindings)
+	{
+		DECLARE_SCOPE_CYCLE_COUNTER(TEXT("FModelCPU_Run"), STAT_FModelCPU_Run, STATGROUP_NNE);
+
+		// Sanity check
+		if (!bIsLoaded)
+		{
+			UE_LOG(LogNNE, Warning, TEXT("FModelCPU::Run(): Call FModelCPU::Load() to load a model first."));
+			return -1;
+		}
+
+		// Verify the model inputs were prepared
+		if (InputTensorShapes.Num() == 0)
+		{
+			UE_LOG(LogNNE, Error, TEXT("Run(): Input shapes are not set, please call SetInputTensorShapes."));
+			return -1;
+		}
+
+		NNEProfiling::Internal::FTimer RunTimer;
+		RunTimer.Tic();
+
+		if (!bHasRun)
+		{
+			bHasRun = true;
+		}
+
+#if WITH_EDITOR
+		try
+#endif //WITH_EDITOR
+		{
+			TArray<Ort::Value> InputOrtTensors;
+			BindTensorsToORT(InInputBindings, InputTensors, InputTensorsORTType, AllocatorInfo.Get(), InputOrtTensors);
+
+			if (!OutputTensors.IsEmpty())
+			{
+				// If output shapes are known we can directly map preallocated output buffers
+				TArray<Ort::Value> OutputOrtTensors;
+				BindTensorsToORT(InOutputBindings, OutputTensors, OutputTensorsORTType, AllocatorInfo.Get(), OutputOrtTensors);
+
+				Session->Run(Ort::RunOptions{ nullptr },
+					InputTensorNames.GetData(), &InputOrtTensors[0], InputTensorNames.Num(),
+					OutputTensorNames.GetData(), &OutputOrtTensors[0], OutputTensorNames.Num());
+			}
+			else
+			{
+				TArray<Ort::Value> OutputOrtTensors;
+				for (int i = 0; i < InOutputBindings.Num(); ++i)
+				{
+					OutputOrtTensors.Emplace(nullptr);
+				}
+
+				Session->Run(Ort::RunOptions{ nullptr },
+					InputTensorNames.GetData(), &InputOrtTensors[0], InputTensorNames.Num(),
+					OutputTensorNames.GetData(), &OutputOrtTensors[0], OutputTensorNames.Num());
+
+				// Output shapes were resolved during inference: Copy the data back to bindings and expose output tensor shapes
+				CopyFromORTToBindings(OutputOrtTensors, InOutputBindings, OutputSymbolicTensors, OutputTensors);
+				check(OutputTensorShapes.IsEmpty());
+				for (int i = 0; i < OutputTensors.Num(); ++i)
+				{
+					OutputTensorShapes.Emplace(OutputTensors[i].GetShape());
+				}
+			}
+		}
+#if WITH_EDITOR
+		catch (const std::exception& Exception)
+		{
+			UE_LOG(LogNNE, Error, TEXT("%s"), UTF8_TO_TCHAR(Exception.what()));
+		}
+#endif //WITH_EDITOR
+
+		RunStatisticsEstimator.StoreSample(RunTimer.Toc());
+
+		return 0;
+	}
+
+	float FModelCPU::GetLastRunTimeMSec() const
+	{
+		return RunStatisticsEstimator.GetLastSample();
+	}
+
+	NNEProfiling::Internal::FStatistics FModelCPU::GetRunStatistics() const
+	{
+		return RunStatisticsEstimator.GetStats();
+	}
+
+	NNEProfiling::Internal::FStatistics FModelCPU::GetInputMemoryTransferStats() const
+	{
+		return InputTransferStatisticsEstimator.GetStats();
+	}
+
+	void FModelCPU::ResetStats()
+	{
+		RunStatisticsEstimator.ResetStats();
+		InputTransferStatisticsEstimator.ResetStats();
+	}
+
+} // namespace UE::NNERuntimeORTCpu::Private
\ No newline at end of file
diff --git a/Engine/Plugins/Experimental/NNE/Source/NNERuntimeORTCpu/Private/NNERuntimeORTCpuModel.h b/Engine/Plugins/Experimental/NNE/Source/NNERuntimeORTCpu/Private/NNERuntimeORTCpuModel.h
new file mode 100644
--- /dev/null
+++ b/Engine/Plugins/Experimental/NNE/Source/NNERuntimeORTCpu/Private/NNERuntimeORTCpuModel.h
@@ -0,0 +1,81 @@
+// Copyright Epic Games, Inc. All Rights Reserved.
+
+#pragma once
+
+#include "CoreMinimal.h"
+#include "NNECoreModelBase.h"
+#include "NNECoreRuntimeCPU.h"
+#include "NNECoreTensor.h"
+#include "NNECoreTypes.h"
+#include "NNEProfilingStatistics.h"
+
+#include "NNXThirdPartyWarningDisabler.h"
+NNX_THIRD_PARTY_INCLUDES_START
+#undef check
+#undef TEXT
+#include "core/session/onnxruntime_cxx_api.h"
+NNX_THIRD_PARTY_INCLUDES_END
+
+namespace UE::NNERuntimeORTCpu::Private
+{
+	struct FRuntimeConf
+	{
+		uint32 NumberOfThreads = 2;
+		GraphOptimizationLevel OptimizationLevel = GraphOptimizationLevel::ORT_ENABLE_ALL;
+		EThreadPriority ThreadPriority = EThreadPriority::TPri_Normal;
+	};
+
+	class FModelCPU : public NNECore::Internal::FModelBase<NNECore::IModelCPU>
+	{
+
+	public:
+		FModelCPU();
+		FModelCPU(Ort::Env* InORTEnvironment, const FRuntimeConf& InRuntimeConf);
+		virtual ~FModelCPU() {};
+
+		virtual int SetInputTensorShapes(TConstArrayView<NNECore::FTensorShape> InInputShapes) override;
+		virtual int RunSync(TConstArrayView<NNECore::FTensorBindingCPU> InInputBindings, TConstArrayView<NNECore::FTensorBindingCPU> InOutputBindings) override;
+
+		bool Init(TConstArrayView<uint8> ModelData);
+		bool IsLoaded() const;
+		float GetLastRunTimeMSec() const;
+		NNEProfiling::Internal::FStatistics GetRunStatistics() const;
+		NNEProfiling::Internal::FStatistics GetInputMemoryTransferStats() const;
+		void ResetStats();
+
+	private:
+
+		bool bIsLoaded = false;
+		bool bHasRun = false;
+
+		FRuntimeConf RuntimeConf;
+
+		/** ORT-related variables */
+		Ort::Env* ORTEnvironment;
+		TUniquePtr<Ort::Session> Session;
+		TUniquePtr<Ort::AllocatorWithDefaultOptions> Allocator;
+		TUniquePtr<Ort::SessionOptions> SessionOptions;
+		TUniquePtr<Ort::MemoryInfo> AllocatorInfo;
+
+		/** IO ORT-related variables */
+		TArray<ONNXTensorElementDataType> InputTensorsORTType;
+		TArray<ONNXTensorElementDataType> OutputTensorsORTType;
+
+		TArray<const char*> InputTensorNames;
+		TArray<const char*> OutputTensorNames;
+
+		TArray<NNECore::Internal::FTensor> InputTensors;
+		TArray<NNECore::Internal::FTensor> OutputTensors;
+
+		bool InitializedAndConfigureMembers();
+		bool ConfigureTensors(const bool InIsInput);
+
+		/**
+		 * Statistics-related members used for GetLastRunTimeMSec(), GetRunStatistics(), GetInputMemoryTransferStats(), ResetStats().
+		 */
+		UE::NNEProfiling::Internal::FStatisticsEstimator RunStatisticsEstimator;
+		UE::NNEProfiling::Internal::FStatisticsEstimator InputTransferStatisticsEstimator;
+
+	};
+	
+} // UE::NNERuntimeORTCpu::Private
\ No newline at end of file
diff --git a/Engine/Plugins/Experimental/NNE/Source/NNERuntimeORTCpu/Private/NNERuntimeORTCpuModule.cpp b/Engine/Plugins/Experimental/NNE/Source/NNERuntimeORTCpu/Private/NNERuntimeORTCpuModule.cpp
new file mode 100644
--- /dev/null
+++ b/Engine/Plugins/Experimental/NNE/Source/NNERuntimeORTCpu/Private/NNERuntimeORTCpuModule.cpp
@@ -0,0 +1,34 @@
+// Copyright Epic Games, Inc. All Rights Reserved.
+
+#include "NNERuntimeORTCpuModule.h"
+#include "NNECore.h"
+#include "NNERuntimeORTCpu.h"
+#include "UObject/WeakInterfacePtr.h"
+
+void FNNERuntimeORTCpuModule::StartupModule()
+{
+	// NNE runtime startup
+	NNERuntimeORTCpu = NewObject<UNNERuntimeORTCpuImpl>();
+	if (NNERuntimeORTCpu.IsValid())
+	{
+		TWeakInterfacePtr<INNERuntime> RuntimeCPUInterface(NNERuntimeORTCpu.Get());
+		
+		NNERuntimeORTCpu->AddToRoot();
+		UE::NNECore::RegisterRuntime(RuntimeCPUInterface);
+	}
+}
+
+void FNNERuntimeORTCpuModule::ShutdownModule()
+{
+	// NNE runtime shutdown
+	if (NNERuntimeORTCpu.IsValid())
+	{
+		TWeakInterfacePtr<INNERuntime> RuntimeCPUInterface(NNERuntimeORTCpu.Get());
+		
+		UE::NNECore::UnregisterRuntime(RuntimeCPUInterface);
+		NNERuntimeORTCpu->RemoveFromRoot();
+		NNERuntimeORTCpu = TWeakObjectPtr<UNNERuntimeORTCpuImpl>(nullptr);
+	}
+}
+
+IMPLEMENT_MODULE(FNNERuntimeORTCpuModule, NNERuntimeORTCpu);
diff --git a/Engine/Plugins/Experimental/NNE/Source/NNERuntimeORTCpu/Private/NNERuntimeORTCpuModule.h b/Engine/Plugins/Experimental/NNE/Source/NNERuntimeORTCpu/Private/NNERuntimeORTCpuModule.h
new file mode 100644
--- /dev/null
+++ b/Engine/Plugins/Experimental/NNE/Source/NNERuntimeORTCpu/Private/NNERuntimeORTCpuModule.h
@@ -0,0 +1,23 @@
+// Copyright Epic Games, Inc. All Rights Reserved.
+
+#pragma once
+
+#include "CoreMinimal.h"
+#include "Modules/ModuleManager.h"
+#include "UObject/WeakObjectPtr.h"
+#include "UObject/WeakObjectPtrTemplates.h"
+
+namespace NNX { class IRuntime; }
+class UNNERuntimeORTCpuImpl;
+
+class FNNERuntimeORTCpuModule : public IModuleInterface
+{
+
+public:
+	NNX::IRuntime* CPURuntime{ nullptr };
+	TWeakObjectPtr<UNNERuntimeORTCpuImpl> NNERuntimeORTCpu{ nullptr };
+
+	// Begin IModuleInterface
+	virtual void StartupModule() override;
+	virtual void ShutdownModule() override;
+};
\ No newline at end of file
diff --git a/Engine/Plugins/Experimental/NNE/Source/NNERuntimeORTCpu/Private/NNERuntimeORTCpuUtils.h b/Engine/Plugins/Experimental/NNE/Source/NNERuntimeORTCpu/Private/NNERuntimeORTCpuUtils.h
new file mode 100644
--- /dev/null
+++ b/Engine/Plugins/Experimental/NNE/Source/NNERuntimeORTCpu/Private/NNERuntimeORTCpuUtils.h
@@ -0,0 +1,214 @@
+// Copyright Epic Games, Inc. All Rights Reserved.
+
+#pragma once
+
+#include "HAL/UnrealMemory.h"
+#include "NNECore.h"
+#include "NNECoreRuntimeCPU.h"
+#include "NNECoreTypes.h"
+#include "NNECoreTensor.h"
+#include "NNXThirdPartyWarningDisabler.h"
+NNX_THIRD_PARTY_INCLUDES_START
+#undef check
+#undef TEXT
+#include "core/session/onnxruntime_cxx_api.h"
+NNX_THIRD_PARTY_INCLUDES_END
+
+DECLARE_STATS_GROUP(TEXT("NNE"), STATGROUP_NNE, STATCAT_Advanced);
+
+namespace UE::NNERuntimeORTCpu::Private
+{
+
+	using TypeInfoORT = std::pair<ENNETensorDataType, uint64>;
+
+	inline TypeInfoORT TranslateTensorTypeORTToNNE(unsigned int OrtDataType)
+	{
+		ENNETensorDataType DataType = ENNETensorDataType::None;
+		uint64 ElementSize = 0;
+
+		switch (OrtDataType) {
+		case ONNX_TENSOR_ELEMENT_DATA_TYPE_UNDEFINED:
+			DataType = ENNETensorDataType::None;
+			ElementSize = 0;
+			break;
+
+		case ONNX_TENSOR_ELEMENT_DATA_TYPE_FLOAT:
+			DataType = ENNETensorDataType::Float;
+			ElementSize = sizeof(float);
+			break;
+
+		case ONNX_TENSOR_ELEMENT_DATA_TYPE_UINT8:
+			DataType = ENNETensorDataType::UInt8;
+			ElementSize = sizeof(uint8);
+			break;
+		case ONNX_TENSOR_ELEMENT_DATA_TYPE_INT8:
+			DataType = ENNETensorDataType::Int8;
+			ElementSize = sizeof(int8);
+			break;
+		case ONNX_TENSOR_ELEMENT_DATA_TYPE_UINT16:
+			DataType = ENNETensorDataType::UInt16;
+			ElementSize = sizeof(uint16);
+			break;
+		case ONNX_TENSOR_ELEMENT_DATA_TYPE_INT16:
+			DataType = ENNETensorDataType::Int16;
+			ElementSize = sizeof(int16);
+			break;
+
+		case ONNX_TENSOR_ELEMENT_DATA_TYPE_INT32:
+			DataType = ENNETensorDataType::Int32;
+			ElementSize = sizeof(int32);
+			break;
+
+		case ONNX_TENSOR_ELEMENT_DATA_TYPE_INT64:
+			DataType = ENNETensorDataType::Int64;
+			ElementSize = sizeof(int64);
+			break;
+
+		case ONNX_TENSOR_ELEMENT_DATA_TYPE_STRING:
+			DataType = ENNETensorDataType::Char;
+			ElementSize = sizeof(char);
+			break;
+
+		case ONNX_TENSOR_ELEMENT_DATA_TYPE_BOOL:
+			DataType = ENNETensorDataType::Boolean;
+			ElementSize = sizeof(bool);
+			break;
+
+		case ONNX_TENSOR_ELEMENT_DATA_TYPE_FLOAT16:
+			DataType = ENNETensorDataType::Half;
+			ElementSize = 2;
+			break;
+
+		case ONNX_TENSOR_ELEMENT_DATA_TYPE_DOUBLE:
+			DataType = ENNETensorDataType::Double;
+			ElementSize = sizeof(double);
+			break;
+
+		case ONNX_TENSOR_ELEMENT_DATA_TYPE_UINT32:
+			DataType = ENNETensorDataType::UInt32;
+			ElementSize = sizeof(uint32);
+			break;
+
+		case ONNX_TENSOR_ELEMENT_DATA_TYPE_UINT64:
+			DataType = ENNETensorDataType::UInt64;
+			ElementSize = sizeof(uint64);
+			break;
+
+		case ONNX_TENSOR_ELEMENT_DATA_TYPE_COMPLEX64:
+			DataType = ENNETensorDataType::Complex64;
+			ElementSize = 8;
+			break;
+
+		case ONNX_TENSOR_ELEMENT_DATA_TYPE_COMPLEX128:
+			DataType = ENNETensorDataType::Complex128;
+			ElementSize = 16;
+			break;
+
+		case ONNX_TENSOR_ELEMENT_DATA_TYPE_BFLOAT16:
+			DataType = ENNETensorDataType::BFloat16;
+			ElementSize = 2;
+			break;
+
+		default:
+			DataType = ENNETensorDataType::None;
+			break;
+		}
+
+		return TypeInfoORT{ DataType, ElementSize };
+	}
+
+	inline void BindTensorsToORT(
+		TConstArrayView<NNECore::FTensorBindingCPU> InBindingTensors,
+		TConstArrayView<NNECore::Internal::FTensor> InTensors,
+		TConstArrayView<ONNXTensorElementDataType> InTensorsORTType,
+		const Ort::MemoryInfo* InAllocatorInfo,
+		TArray<Ort::Value>& OutOrtTensors
+	)
+	{
+		const uint32 NumBinding = InBindingTensors.Num();
+		const uint32 NumDescriptors = InTensors.Num();
+
+		if (NumBinding != NumDescriptors)
+		{
+			UE_LOG(LogNNE, Error, TEXT("BindTensorsToORT: Number of Bindings is different of Descriptors."));
+			return;
+		}
+
+		for (uint32 Index = 0; Index < NumBinding; ++Index)
+		{
+			const NNECore::FTensorBindingCPU& Binding = InBindingTensors[Index];
+			const NNECore::Internal::FTensor& Tensor = InTensors[Index];
+			const ONNXTensorElementDataType CurrentORTType = InTensorsORTType[Index];
+
+			TUniquePtr<int64_t[]> SizesInt64t;
+			SizesInt64t = MakeUnique<int64_t[]>(Tensor.GetShape().Rank());
+			for (int32 DimIndex = 0; DimIndex < Tensor.GetShape().Rank(); ++DimIndex)
+			{
+				SizesInt64t.Get()[DimIndex] = Tensor.GetShape().GetData()[DimIndex];
+			}
+
+			const uint64 ByteCount{ InTensors[Index].GetDataSize() };
+			const uint32 ArrayDimensions{ (uint32)Tensor.GetShape().Rank() };
+			OutOrtTensors.Emplace(
+				Ort::Value::CreateTensor(
+					*InAllocatorInfo,
+					Binding.Data,
+					ByteCount,
+					SizesInt64t.Get(),
+					ArrayDimensions,
+					CurrentORTType
+				)
+			);
+		}
+	}
+
+	inline void CopyFromORTToBindings(
+		TConstArrayView<Ort::Value> InOrtTensors,
+		TConstArrayView<NNECore::FTensorBindingCPU> InBindingTensors,
+		TConstArrayView<NNECore::FTensorDesc> InTensorDescs,
+		TArray<NNECore::Internal::FTensor>& OutTensors
+	)
+	{
+		const uint32 NumBinding = InOrtTensors.Num();
+		const uint32 NumDescriptors = InTensorDescs.Num();
+
+		if (NumBinding != NumDescriptors)
+		{
+			UE_LOG(LogNNE, Error, TEXT("CopyFromORTToBindings: Number of Bindings is different of Descriptors."));
+			return;
+		}
+
+		for (uint32 Index = 0; Index < NumBinding; ++Index)
+		{
+			const NNECore::FTensorBindingCPU& Binding = InBindingTensors[Index];
+			const NNECore::FTensorDesc& TensorDesc = InTensorDescs[Index];
+			const Ort::Value& OrtTensor = InOrtTensors[Index];
+			const std::vector<int64_t>& OrtShape = OrtTensor.GetTensorTypeAndShapeInfo().GetShape();
+
+			TArray<uint32> ShapeData;
+			for (int32 DimIndex = 0; DimIndex < OrtShape.size(); ++DimIndex)
+			{
+				check(OrtShape[DimIndex] >= 0);
+				ShapeData.Add(OrtShape[DimIndex]);
+			}
+
+			NNECore::FTensorShape Shape = NNECore::FTensorShape::Make(ShapeData);
+			NNECore::Internal::FTensor Tensor = NNECore::Internal::FTensor::Make(TensorDesc.GetName(), Shape, TensorDesc.GetDataType());
+			OutTensors.Emplace(Tensor);
+
+			const void* OrtTensorData = OrtTensor.GetTensorData<void>();
+			const uint64 DataSize = FMath::Min(Tensor.GetDataSize(), Binding.SizeInBytes);
+
+			if (DataSize > 0)
+			{
+				FMemory::Memcpy(Binding.Data, OrtTensorData, DataSize);
+			}
+
+			if (Tensor.GetDataSize() > Binding.SizeInBytes)
+			{
+				UE_LOG(LogNNE, Error, TEXT("CopyFromORTToBindings: Binding buffer was not large enough to contain all of the data, only first %d bytes were copied."), DataSize);
+			}
+		}
+	}
+
+} // UE::NNERuntimeORTCpu::Private
\ No newline at end of file
diff --git a/Engine/Plugins/Experimental/NNE/Source/NNXRuntimeCPU/NNXRuntimeCPU.Build.cs b/Engine/Plugins/Experimental/NNE/Source/NNXRuntimeCPU/NNXRuntimeCPU.Build.cs
deleted file mode 100644
diff --git a/Engine/Plugins/Experimental/NNE/Source/NNXRuntimeCPU/Private/NNERuntimeCPU.cpp b/Engine/Plugins/Experimental/NNE/Source/NNXRuntimeCPU/Private/NNERuntimeCPU.cpp
deleted file mode 100644
diff --git a/Engine/Plugins/Experimental/NNE/Source/NNXRuntimeCPU/Private/NNERuntimeCPU.h b/Engine/Plugins/Experimental/NNE/Source/NNXRuntimeCPU/Private/NNERuntimeCPU.h
deleted file mode 100644
diff --git a/Engine/Plugins/Experimental/NNE/Source/NNXRuntimeCPU/Private/NNERuntimeCPUModel.cpp b/Engine/Plugins/Experimental/NNE/Source/NNXRuntimeCPU/Private/NNERuntimeCPUModel.cpp
deleted file mode 100644
diff --git a/Engine/Plugins/Experimental/NNE/Source/NNXRuntimeCPU/Private/NNERuntimeCPUModel.h b/Engine/Plugins/Experimental/NNE/Source/NNXRuntimeCPU/Private/NNERuntimeCPUModel.h
deleted file mode 100644
diff --git a/Engine/Plugins/Experimental/NNE/Source/NNXRuntimeCPU/Private/NNERuntimeCPUModule.cpp b/Engine/Plugins/Experimental/NNE/Source/NNXRuntimeCPU/Private/NNERuntimeCPUModule.cpp
deleted file mode 100644
diff --git a/Engine/Plugins/Experimental/NNE/Source/NNXRuntimeCPU/Private/NNERuntimeCPUModule.h b/Engine/Plugins/Experimental/NNE/Source/NNXRuntimeCPU/Private/NNERuntimeCPUModule.h
deleted file mode 100644
diff --git a/Engine/Plugins/Experimental/NNE/Source/NNXRuntimeCPU/Private/NNERuntimeCPUUtils.h b/Engine/Plugins/Experimental/NNE/Source/NNXRuntimeCPU/Private/NNERuntimeCPUUtils.h
deleted file mode 100644
diff --git a/Engine/Plugins/Experimental/NNE/Source/NNXRuntimeCPU/Private/NNXRuntimeCPU.cpp b/Engine/Plugins/Experimental/NNE/Source/NNXRuntimeCPU/Private/NNXRuntimeCPU.cpp
deleted file mode 100644
diff --git a/Engine/Plugins/Experimental/NNE/Source/NNXRuntimeCPU/Public/NNXRuntimeCPU.h b/Engine/Plugins/Experimental/NNE/Source/NNXRuntimeCPU/Public/NNXRuntimeCPU.h
deleted file mode 100644
diff --git a/Engine/Plugins/Experimental/NNE/Source/NNXRuntimeCPU/Public/NNXRuntimeCPUUtils.h b/Engine/Plugins/Experimental/NNE/Source/NNXRuntimeCPU/Public/NNXRuntimeCPUUtils.h
deleted file mode 100644
