NNE - Changing CPUx to use the minimal API

diff --git a/Engine/Plugins/Experimental/NNE/Source/NNERuntimeORTCpu/Private/NNERuntimeORTCpuModule.h b/Engine/Plugins/Experimental/NNE/Source/NNERuntimeORTCpu/Private/NNERuntimeORTCpuModule.h
--- a/Engine/Plugins/Experimental/NNE/Source/NNERuntimeORTCpu/Private/NNERuntimeORTCpuModule.h
+++ b/Engine/Plugins/Experimental/NNE/Source/NNERuntimeORTCpu/Private/NNERuntimeORTCpuModule.h
@@ -7,14 +7,12 @@
 #include "UObject/WeakObjectPtr.h"
 #include "UObject/WeakObjectPtrTemplates.h"
 
-namespace NNX { class IRuntime; }
 class UNNERuntimeORTCpuImpl;
 
 class FNNERuntimeORTCpuModule : public IModuleInterface
 {
 
 public:
-	NNX::IRuntime* CPURuntime{ nullptr };
 	TWeakObjectPtr<UNNERuntimeORTCpuImpl> NNERuntimeORTCpu{ nullptr };
 
 	// Begin IModuleInterface
