NNE : Minimal API - NNERuntimeRDG

Problem:
* NNE API surface is large atm potentially creating a support burden and harming adoption for 5.2.

Solution:
* This CL is targeting NNERuntimeRDG to use the minimal API.

Test:
* All tests green under win64 editor

diff --git a/Engine/Plugins/Experimental/NNE/Source/NNEQA/NNEQA.Build.cs b/Engine/Plugins/Experimental/NNE/Source/NNEQA/NNEQA.Build.cs
--- a/Engine/Plugins/Experimental/NNE/Source/NNEQA/NNEQA.Build.cs
+++ b/Engine/Plugins/Experimental/NNE/Source/NNEQA/NNEQA.Build.cs
@@ -23,7 +23,9 @@ public class NNEQA : ModuleRules
 				"Json",
 				"CoreUObject",
 				"NNXCore",
-				"NNEUtils"
+				"NNEUtils",
+				"RHI",
+				"RenderCore"
 			}
 		);
 
diff --git a/Engine/Plugins/Experimental/NNE/Source/NNEQA/Private/NNEQAModel.cpp b/Engine/Plugins/Experimental/NNE/Source/NNEQA/Private/NNEQAModel.cpp
new file mode 100644
--- /dev/null
+++ b/Engine/Plugins/Experimental/NNE/Source/NNEQA/Private/NNEQAModel.cpp
@@ -0,0 +1,279 @@
+// Copyright Epic Games, Inc. All Rights Reserved.
+
+#include "NNEQAModel.h"
+#include "NNECoreModelData.h"
+#include "RenderGraphResources.h"
+#include "ShaderParameterUtils.h"
+
+BEGIN_SHADER_PARAMETER_STRUCT(FNNEQATensorReadbackParameters, )
+	RDG_BUFFER_ACCESS(Buffer, ERHIAccess::CopySrc)
+END_SHADER_PARAMETER_STRUCT()
+
+namespace UE::NNEQA::Private 
+{
+namespace FModelQAHelpers
+{
+	FRDGBufferDesc CreateRDGBufferDescForTensorDesc(uint32 ElemByteSize, uint64 SizeInByte)
+	{
+		// FIXME: CreateStructuredDesc() creates a crash on VulkanRHI
+		//FRDGBufferDesc Desc = FRDGBufferDesc::CreateStructuredDesc(ElemByteSize, SizeInByte/ElemByteSize);
+		FRDGBufferDesc Desc = FRDGBufferDesc::CreateBufferDesc(ElemByteSize, SizeInByte/ElemByteSize);
+
+		// FIXME: We should use BUF_SourceCopy for only output buffers (GPU readback)
+		Desc.Usage = EBufferUsageFlags(Desc.Usage | BUF_SourceCopy);
+
+		return Desc;
+	}
+
+	void ConvertBinding(FRDGBuilder& GraphBuilder,
+		TConstArrayView<NNECore::FTensorDesc> InTensorDescs,
+		TConstArrayView<NNECore::FTensorBindingCPU> InBindingsCPU,
+		TArray<NNECore::FTensorBindingRDG>& OutBindingsRDG)
+	{
+		check(IsInRenderingThread());
+		check(InTensorDescs.Num() == InBindingsCPU.Num());
+
+		for (int32 Idx = 0; Idx < InBindingsCPU.Num(); ++Idx)
+		{
+			const NNECore::FTensorDesc& TensorDesc = InTensorDescs[Idx];
+			const NNECore::FTensorBindingCPU& BindingCPU = InBindingsCPU[Idx];
+			const FRDGBufferDesc Desc = CreateRDGBufferDescForTensorDesc(TensorDesc.GetElemByteSize(), BindingCPU.SizeInBytes);
+			FRDGBufferRef TensorBuffer = GraphBuilder.CreateBuffer(Desc, *TensorDesc.GetName(), ERDGBufferFlags::None);
+
+			NNECore::FTensorBindingRDG& BindingRDG = OutBindingsRDG.Emplace_GetRef();
+			BindingRDG.Buffer = TensorBuffer;
+		}
+	}
+
+	void UploadsBindingToGPU(FRDGBuilder& GraphBuilder,
+		TConstArrayView<NNECore::FTensorBindingCPU> InBindingsCPU,
+		TConstArrayView<NNECore::FTensorBindingRDG> InBindingsRDG)
+	{
+		check(IsInRenderingThread());
+		check(InBindingsCPU.Num() == InBindingsRDG.Num());
+
+		for (int32 Idx = 0; Idx < InBindingsCPU.Num(); ++Idx)
+		{
+			const NNECore::FTensorBindingCPU& BindingCPU = InBindingsCPU[Idx];
+			const NNECore::FTensorBindingRDG& BindingRDG = InBindingsRDG[Idx];
+			GraphBuilder.QueueBufferUpload(BindingRDG.Buffer, BindingCPU.Data, BindingCPU.SizeInBytes, ERDGInitialDataFlags::NoCopy);
+		}
+	}
+
+	void DownloadBindingToCPU(FRDGBuilder& GraphBuilder,
+		TArray<FModelQA::FReadbackEntry>& Readbacks,
+		TConstArrayView<NNECore::FTensorBindingRDG> InBindingsRDG,
+		TConstArrayView<NNECore::FTensorBindingCPU> InBindingsCPU)
+	{
+		check(IsInRenderingThread());
+		check(InBindingsCPU.Num() == InBindingsRDG.Num());
+
+		for (int32 Idx = 0; Idx < InBindingsRDG.Num(); ++Idx)
+		{
+			const NNECore::FTensorBindingCPU& BindingCPU = InBindingsCPU[Idx];
+			const NNECore::FTensorBindingRDG& BindingRDG = InBindingsRDG[Idx];
+			
+			FNNEQATensorReadbackParameters* TensorReadbackParams = GraphBuilder.AllocParameters<FNNEQATensorReadbackParameters>();
+
+			TensorReadbackParams->Buffer = BindingRDG.Buffer;
+
+			GraphBuilder.AddPass(
+				RDG_EVENT_NAME("FNNEQAModelTensorReadback"),
+				TensorReadbackParams,
+				ERDGPassFlags::Readback | ERDGPassFlags::NeverCull,
+				[&Readbacks, BindingCPU, TensorReadbackParams](FRHICommandListImmediate& RHICmdList)
+				{
+					FRHIBuffer* OutputBuffer = TensorReadbackParams->Buffer->GetRHI();
+
+					// TODO: FIXME: We need to transition the resources for DirectML
+					/*if (bUseManualTransitions)
+					{
+						FRHITransitionInfo Transitions[] =
+						{
+							FRHITransitionInfo(OutputBuffer, ERHIAccess::UAVCompute, ERHIAccess::CopySrc)
+						};
+
+						RHICmdList.Transition(MakeArrayView(Transitions, UE_ARRAY_COUNT(Transitions)));
+						RHICmdList.SubmitCommandsHint();
+					}*/
+
+					FModelQA::FReadbackEntry& Readback = Readbacks.Emplace_GetRef();
+					Readback.RHI = MakeUnique<FRHIGPUBufferReadback>(FName(TEXT("FNNEQAModelTensorReadback")));
+					Readback.RHI->EnqueueCopy(RHICmdList, OutputBuffer, BindingCPU.SizeInBytes);
+					Readback.CpuMemory = BindingCPU.Data;
+					Readback.Size = BindingCPU.SizeInBytes;
+				}
+			);
+		}
+	}
+}
+	
+TConstArrayView<NNECore::FTensorDesc> FModelQA::GetInputTensorDescs() const
+{ 
+	if (ModelCPU.IsValid()) 
+		return ModelCPU->GetInputTensorDescs();
+	if (ModelRDG.IsValid())
+		return ModelRDG->GetInputTensorDescs();
+	else
+		return {};
+};
+		
+TConstArrayView<NNECore::FTensorDesc> FModelQA::GetOutputTensorDescs() const
+{
+	if (ModelCPU.IsValid())
+		return ModelCPU->GetOutputTensorDescs();
+	if (ModelRDG.IsValid())
+		return ModelRDG->GetOutputTensorDescs();
+	else
+		return {};
+}
+
+TConstArrayView<NNECore::FTensorShape> FModelQA::GetInputTensorShapes() const
+{
+	if (ModelCPU.IsValid())
+		return ModelCPU->GetInputTensorShapes();
+	if (ModelRDG.IsValid())
+		return ModelRDG->GetInputTensorShapes();
+	else
+		return {};
+}
+
+TConstArrayView<NNECore::FTensorShape> FModelQA::GetOutputTensorShapes() const
+{
+	if (ModelCPU.IsValid())
+		return ModelCPU->GetOutputTensorShapes();
+	if (ModelRDG.IsValid())
+		return ModelRDG->GetOutputTensorShapes();
+	else
+		return {};
+}
+
+int FModelQA::SetInputTensorShapes(TConstArrayView<NNECore::FTensorShape> InInputShapes)
+{
+	if (ModelCPU.IsValid())
+		return ModelCPU->SetInputTensorShapes(InInputShapes);
+	if (ModelRDG.IsValid())
+		return ModelRDG->SetInputTensorShapes(InInputShapes);
+	else
+		return {};
+}
+	
+int FModelQA::EnqueueRDG(FRDGBuilder& RDGBuilder,
+	TConstArrayView<NNECore::FTensorBindingCPU> InInputBindings, 
+	TConstArrayView<NNECore::FTensorBindingCPU> InOutputBindings)
+{
+	check(IsInRenderingThread());
+			
+	if (!ModelRDG.IsValid())
+	{
+		return -1;
+	}
+			
+	TArray<NNECore::FTensorBindingRDG> InputBindingsRDG;
+	TArray<NNECore::FTensorBindingRDG> OutputBindingsRDG;
+			
+	FModelQAHelpers::ConvertBinding(RDGBuilder, ModelRDG->GetInputTensorDescs(), InInputBindings, InputBindingsRDG);
+	FModelQAHelpers::ConvertBinding(RDGBuilder, ModelRDG->GetOutputTensorDescs(), InOutputBindings, OutputBindingsRDG);
+			
+	FModelQAHelpers::UploadsBindingToGPU(RDGBuilder, InInputBindings, InputBindingsRDG);
+
+	int32 Res = ModelRDG->EnqueueRDG(RDGBuilder, InputBindingsRDG, OutputBindingsRDG);
+
+	FModelQAHelpers::DownloadBindingToCPU(RDGBuilder, Readbacks, OutputBindingsRDG, InOutputBindings);
+			
+	return Res;
+}
+
+int FModelQA::RunSync(TConstArrayView<NNECore::FTensorBindingCPU> InInputBindings, TConstArrayView<NNECore::FTensorBindingCPU> InOutputBindings)
+{
+	if (ModelCPU.IsValid())
+	{
+		return ModelCPU->RunSync(InInputBindings, InOutputBindings);
+	}
+	else if (ModelRDG.IsValid())
+	{
+		// TODO ok if pending?
+		Readbacks.Empty();
+
+		int Res = 0;
+		FEvent* Signal = FGenericPlatformProcess::GetSynchEventFromPool(false);
+
+		ENQUEUE_RENDER_COMMAND(FModelQA_Run)
+		(
+			[&Signal, &Res, this, InInputBindings, InOutputBindings](FRHICommandListImmediate& RHICmdList)
+			{
+				TOptional<ERHIPipeline> Pipeline = RHICmdList.GetPipeline();
+
+				if (Pipeline == ERHIPipeline::None)
+				{
+					RHICmdList.SwitchPipeline(ERHIPipeline::Graphics);
+				}
+
+				FRDGBuilder	RDGBuilder(RHICmdList);
+
+				Res = EnqueueRDG(RDGBuilder, InInputBindings, InOutputBindings);
+				if (Res == 0)
+				{
+					RDGBuilder.Execute();
+
+					// FIXME: Using BlockUntilGPUIdle() prevents hang on Linux
+					// FIXME: Adapt to redesigned readback API (UE 5.2)
+					RHICmdList.BlockUntilGPUIdle();
+
+					// Process readback
+					for (const FReadbackEntry& Readback : Readbacks) {
+						const void* BuffData = Readback.RHI->Lock(Readback.Size);
+						check(BuffData);
+						FMemory::Memcpy(Readback.CpuMemory, BuffData, Readback.Size);
+						Readback.RHI->Unlock();
+					}
+				}
+
+				Signal->Trigger();
+			}
+		);
+
+		// We need to wait for render thread to finish
+		Signal->Wait();
+				
+		FGenericPlatformProcess::ReturnSynchEventToPool(Signal);
+
+		return Res;
+	}
+	else
+	{
+		return -1;
+	}
+}
+
+TUniquePtr<FModelQA> FModelQA::MakeModelQA(const FNNIModelRaw& ONNXModelData, const FString& RuntimeName)
+{
+	TUniquePtr<FModelQA> ModelQA = MakeUnique<FModelQA>();
+	TWeakInterfacePtr<INNERuntime> Runtime = NNECore::GetRuntime<INNERuntime>(RuntimeName);
+	TWeakInterfacePtr<INNERuntimeCPU> RuntimeCPU = NNECore::GetRuntime<INNERuntimeCPU>(RuntimeName);
+	TWeakInterfacePtr<INNERuntimeRDG> RuntimeRDG = NNECore::GetRuntime<INNERuntimeRDG>(RuntimeName);
+	TObjectPtr<UNNEModelData> ModelData = NewObject<UNNEModelData>();
+			
+	ModelData->Init(FString("onnx"), ONNXModelData.Data);
+	if (!Runtime.IsValid())
+	{
+		UE_LOG(LogNNE, Error, TEXT("Can't get %s runtime."), *RuntimeName);
+		return TUniquePtr<FModelQA>();
+	}
+	if (RuntimeCPU.IsValid())
+	{
+		ModelQA->ModelCPU = RuntimeCPU->CreateModelCPU(ModelData);
+	}
+	else if (RuntimeRDG.IsValid())
+	{
+		ModelQA->ModelRDG = RuntimeRDG->CreateModelRDG(ModelData);
+	}
+	else
+	{
+		UE_LOG(LogNNE, Error, TEXT("Can't find supported API for %s runtime."), *RuntimeName);
+		return TUniquePtr<FModelQA>();
+	}
+	return ModelQA;
+}
+	
+} // namespace UE::NNEQA::Private
\ No newline at end of file
diff --git a/Engine/Plugins/Experimental/NNE/Source/NNEQA/Private/NNEQAModel.h b/Engine/Plugins/Experimental/NNE/Source/NNEQA/Private/NNEQAModel.h
new file mode 100644
--- /dev/null
+++ b/Engine/Plugins/Experimental/NNE/Source/NNEQA/Private/NNEQAModel.h
@@ -0,0 +1,56 @@
+// Copyright Epic Games, Inc. All Rights Reserved.
+
+#pragma once
+
+#include "NNEQAModel.h"
+#include "CoreMinimal.h"
+#include "NNECore.h"
+#include "NNECoreRuntimeCPU.h"
+#include "NNECoreRuntimeRDG.h"
+#include "NNECoreTypes.h"
+#include "NNXRuntimeFormat.h"
+#include "RenderGraphBuilder.h"
+#include "RenderGraphUtils.h"
+#include "RHIGPUReadback.h"
+#include "Templates/UniquePtr.h"
+
+namespace UE::NNEQA::Private 
+{
+
+// This class is a wrapper around a NNECore::IModelCPU and NNECore::IModelRDG
+// Allowing to run a model on CPU or GPU using IModelCPU interface only.
+class FModelQA : NNECore::IModelCPU
+{
+private:
+
+	TUniquePtr<NNECore::IModelCPU> ModelCPU;
+	TUniquePtr<NNECore::IModelRDG> ModelRDG;
+
+public:
+	struct FReadbackEntry
+	{
+		TUniquePtr<FRHIGPUBufferReadback> RHI;
+		void* CpuMemory;
+		size_t	Size;
+	};
+	TArray<FReadbackEntry> Readbacks;
+
+	FModelQA() = default;
+	virtual ~FModelQA() = default;
+
+	virtual TConstArrayView<NNECore::FTensorDesc> GetInputTensorDescs() const;		
+	virtual TConstArrayView<NNECore::FTensorDesc> GetOutputTensorDescs() const;
+	virtual TConstArrayView<NNECore::FTensorShape> GetInputTensorShapes() const;
+	virtual TConstArrayView<NNECore::FTensorShape> GetOutputTensorShapes() const;
+	virtual int SetInputTensorShapes(TConstArrayView<NNECore::FTensorShape> InInputShapes);
+	
+	int EnqueueRDG(FRDGBuilder& RDGBuilder,
+		TConstArrayView<NNECore::FTensorBindingCPU> InInputBindings,
+		TConstArrayView<NNECore::FTensorBindingCPU> InOutputBindings);
+
+	virtual int RunSync(TConstArrayView<NNECore::FTensorBindingCPU> InInputBindings, TConstArrayView<NNECore::FTensorBindingCPU> InOutputBindings);
+
+	static TUniquePtr<FModelQA> MakeModelQA(const FNNIModelRaw& ONNXModelData, const FString& RuntimeName);
+};
+	
+} // namespace UE::NNEQA::Private
\ No newline at end of file
diff --git a/Engine/Plugins/Experimental/NNE/Source/NNEQA/Private/NNEQAUtils.cpp b/Engine/Plugins/Experimental/NNE/Source/NNEQA/Private/NNEQAUtils.cpp
--- a/Engine/Plugins/Experimental/NNE/Source/NNEQA/Private/NNEQAUtils.cpp
+++ b/Engine/Plugins/Experimental/NNE/Source/NNEQA/Private/NNEQAUtils.cpp
@@ -1,6 +1,7 @@
 // Copyright Epic Games, Inc. All Rights Reserved.
 
 #include "NNEQAUtils.h"
+#include "NNEQAModel.h"
 #include "HAL/UnrealMemory.h"
 #include "Kismet/GameplayStatics.h"
 #include "NNECore.h"
@@ -8,16 +9,20 @@
 #include "NNECoreModelData.h"
 #include "NNECoreRuntime.h"
 #include "NNECoreRuntimeCPU.h"
+#include "NNECoreRuntimeRDG.h"
 #include "NNXCore.h"
 #include "NNXInferenceModel.h"
 #include "NNXModelOptimizerInterface.h"
 #include "UObject/ObjectPtr.h"
 #include "UObject/WeakInterfacePtr.h"
+#include "RenderGraphBuilder.h"
+#include "RenderGraphUtils.h"
+#include "RHIGPUReadback.h"
 
 namespace UE::NNEQA::Private 
 {
 	using namespace NNX;
-	
+
 	static void FillTensors(
 		TConstArrayView<NNECore::Internal::FTensor> TensorsFromTestSetup,
 		TConstArrayView<NNECore::FTensorDesc> TensorDescsFromModel,
@@ -154,45 +159,6 @@ namespace UE::NNEQA::Private
 		}
 	}
 
-	static void FillInputTensorBindings(TConstArrayView<NNECore::Internal::FTensor> Tensors, TConstArrayView<FTests::FTensorData> TensorsData,
-		TArray<TArray<char>>& OutMemBuffers, TArray<FMLTensorBinding>& OutBindings)
-	{
-		OutBindings.Empty();
-		OutMemBuffers.Empty();
-		check(TensorsData.Num() == 0 || Tensors.Num() == TensorsData.Num());
-
-		for (int Index = 0; Index < Tensors.Num(); ++Index)
-		{
-			const NNECore::Internal::FTensor& Tensor = Tensors[Index];
-			TArray<char>& MemBuffer = OutMemBuffers.Emplace_GetRef();
-			if (TensorsData.Num() == 0 || TensorsData[Index].Num() == 0)
-			{
-				ElementWiseCosTensorInitializer Initializer(Tensor.GetDataType(), Index);
-				MemBuffer = GenerateTensorDataForTest(Tensor, Initializer);
-			}
-			else
-			{
-				MemBuffer = TensorsData[Index];
-			}
-			check(MemBuffer.Num() == Tensor.GetDataSize());
-			OutBindings.Emplace(FMLTensorBinding::FromCPU((void*)(MemBuffer.GetData()), MemBuffer.Num()));
-		}
-	}
-
-	static void FillOutputTensorBindings(TConstArrayView<NNECore::Internal::FTensor> Tensors, TArray<TArray<char>>& OutMemBuffers, TArray<FMLTensorBinding>& OutBindings)
-	{
-		OutBindings.Empty();
-		OutMemBuffers.Empty();
-		for (int Index = 0; Index < Tensors.Num(); ++Index)
-		{
-			TArray<char>& MemBuffer = OutMemBuffers.Emplace_GetRef();
-			MemBuffer.SetNumUninitialized(Tensors[Index].GetDataSize());
-			char constexpr MagicNumber = 0x5b;
-			FMemory::Memset(MemBuffer.GetData(), MagicNumber, MemBuffer.Num());
-			OutBindings.Emplace(FMLTensorBinding::FromCPU((void*)(MemBuffer.GetData()), MemBuffer.Num()));
-		}
-	}
-
 	template<typename T> FString ShapeToString(TConstArrayView<T> Shape)
 	{
 		FString TestSuffix(TEXT("["));
@@ -415,29 +381,12 @@ namespace UE::NNEQA::Private
 		return true;
 	}
 
-	TUniquePtr<NNECore::IModelCPU> GetInferenceModel(const FNNIModelRaw& ONNXModelData, const FString& RuntimeName)
-	{
-		TWeakInterfacePtr<INNERuntime>    RuntimeForModelData      = NNECore::GetRuntime<INNERuntime>(RuntimeName);
-		TWeakInterfacePtr<INNERuntimeCPU> RuntimeForInferenceModel = NNECore::GetRuntime<INNERuntimeCPU>(RuntimeName);
-		if (!RuntimeForModelData.IsValid() || !RuntimeForInferenceModel.IsValid())
-		{
-			UE_LOG(LogNNE, Error, TEXT("Can't get %s runtime."), *RuntimeName);
-			return TUniquePtr<NNECore::IModelCPU>();
-		}
-
-		TObjectPtr<UNNEModelData> ModelData = NewObject<UNNEModelData>();
-		ModelData->Init(FString("onnx"), ONNXModelData.Data);
-		TUniquePtr<NNECore::IModelCPU> InferenceModel = RuntimeForInferenceModel->CreateModelCPU(ModelData);
-		return InferenceModel;
-	}
-	
-	bool RunTestInferenceCPU(const FNNIModelRaw& ONNXModelData, const FTests::FTestSetup& TestSetup,
+	bool RunTestInference(const FNNIModelRaw& ONNXModelData, const FTests::FTestSetup& TestSetup,
 		const FString& RuntimeName, TArray<NNECore::Internal::FTensor>& OutOutputTensors, TArray<TArray<char>>& OutOutputMemBuffers)
 	{
 		OutOutputMemBuffers.Empty();
 
-		TUniquePtr<NNECore::IModelCPU> InferenceModel = GetInferenceModel(ONNXModelData, RuntimeName);
-
+		TUniquePtr<FModelQA> InferenceModel = FModelQA::MakeModelQA(ONNXModelData, RuntimeName);
 		if (!InferenceModel.IsValid())
 		{
 			UE_LOG(LogNNE, Error, TEXT("Could not create Inference model."));
@@ -518,128 +467,6 @@ namespace UE::NNEQA::Private
 		return true;
 	}
 
-	static bool RunTestInference(const FNNIModelRaw& ONNXModelData, const FTests::FTestSetup& TestSetup, 
-		IRuntime* Runtime, TArray<NNECore::Internal::FTensor>& OutOutputTensors, TArray<TArray<char>>& OutOutputMemBuffers)
-	{
-		OutOutputMemBuffers.Empty();
-
-		TUniquePtr<FMLInferenceModel> InferenceModel = Runtime->CreateModel(Runtime->CreateModelData(FString("onnx"), ONNXModelData.Data));
-		
-		if (!InferenceModel.IsValid())
-		{
-			UE_LOG(LogNNE, Error, TEXT("Could not create Inference model."));
-			return false;
-		}
-
-		// If test does not ask for specific input/output, fill with model default converting variable dim to 1 if any.
-		TArray<NNECore::Internal::FTensor> InputTensors;
-		
-		FillTensors(TestSetup.Inputs, InferenceModel->GetInputTensorDescs(), InputTensors);
-		FillTensors(TestSetup.Outputs, InferenceModel->GetOutputTensorDescs(), OutOutputTensors);
-		
-		// bind tensors to memory (CPU) and initialize
-		TArray<FMLTensorBinding> InputBindings;
-		TArray<FMLTensorBinding> OutputBindings;
-		TArray<TArray<char>> InputMemBuffers;
-		
-		FillInputTensorBindings(InputTensors, TestSetup.InputsData, InputMemBuffers, InputBindings);
-		FillOutputTensorBindings(OutOutputTensors, OutOutputMemBuffers, OutputBindings);
-
-		// To help for debugging sessions.
-		#ifdef UE_BUILD_DEBUG
-		const constexpr int32 NumTensorPtrForDebug = 3;
-		float* InputsAsFloat[NumTensorPtrForDebug];
-		float* OutputsAsFloat[NumTensorPtrForDebug];
-		FMemory::Memzero(InputsAsFloat, NumTensorPtrForDebug * sizeof(float*));
-		FMemory::Memzero(OutputsAsFloat, NumTensorPtrForDebug * sizeof(float*));
-		for (int32 i = 0; i < NumTensorPtrForDebug && i < InputMemBuffers.Num(); ++i)
-		{
-			InputsAsFloat[i] = (float*)InputMemBuffers[i].GetData();
-		}
-		for (int32 i = 0; i < NumTensorPtrForDebug && i < OutOutputMemBuffers.Num(); ++i)
-		{
-			OutputsAsFloat[i] = (float*)OutOutputMemBuffers[i].GetData();
-		}
-		#endif //UE_BUILD_DEBUG
-
-		TArray<NNECore::FTensorShape> InputShapes;
-		for (const NNECore::Internal::FTensor& Tensor : InputTensors)
-		{
-			InputShapes.Emplace(Tensor.GetShape());
-		}
-
-		// Setup inputs
-		if (0 != InferenceModel->SetInputTensorShapes(InputShapes))
-		{
-			UE_LOG(LogNNE, Error, TEXT("Failed to set input tensor shapes."));
-			return false;
-		}
-
-		// Run inference
-		if (0 != InferenceModel->RunSync(InputBindings, OutputBindings))
-		{
-			UE_LOG(LogNNE, Error, TEXT("Failed to run the model."));
-			return false;
-		}
-
-		// Verify output shapes are as expected
-		TConstArrayView<NNECore::FTensorShape> OutputShapes = InferenceModel->GetOutputTensorShapes();
-		if (OutputShapes.Num() != OutOutputTensors.Num())
-		{
-			UE_LOG(LogNNE, Error, TEXT("Expected %d output tensors, got %d."), OutOutputTensors.Num(), OutputShapes.Num());
-			return false;
-		}
-		for (int i = 0; i < OutOutputTensors.Num(); ++i)
-		{
-			const NNECore::FTensorShape& RefTensorShape = OutOutputTensors[i].GetShape();
-			const NNECore::FTensorShape& OtherTensorShape = OutputShapes[i];
-			if (RefTensorShape != OtherTensorShape)
-			{
-				TArrayView<const uint32> RefShape(RefTensorShape.GetData());
-				TArrayView<const uint32> OtherShape(OtherTensorShape.GetData());
-				UE_LOG(LogNNE, Error, TEXT("Output shape do not match at index %d.\nExpected: %s\nGot:      %s"), i, *ShapeToString(RefShape), *ShapeToString(OtherShape));
-				return false;
-			}
-		}
-		
-		return true;
-	}
-
-	bool RunTestInferenceAndCompareToRef(const FTests::FTestSetup& TestSetup, IRuntime* Runtime, const FNNIModelRaw& ONNXModel, 
-		TArrayView<TArray<char>> RefOutputMemBuffers, TArrayView<NNECore::Internal::FTensor> RefOutputTensors)
-	{
-		TArray<TArray<char>> OutputMemBuffers;
-		TArray<NNECore::Internal::FTensor> OutputTensors;
-
-		const FString& RuntimeName = Runtime->GetRuntimeName();
-		const float AbsoluteTolerance = TestSetup.GetAbsoluteToleranceForRuntime(RuntimeName);
-		const float RelativeTolerance = TestSetup.GetRelativeToleranceForRuntime(RuntimeName);
-		bool bTestSuceeded = true;
-	
-		if (!RunTestInference(ONNXModel, TestSetup, Runtime, OutputTensors, OutputMemBuffers))
-		{
-			UE_LOG(LogNNE, Error, TEXT("Error running inference for engine %s."), *RuntimeName);
-			return false;
-		}
-
-		if (OutputTensors.Num() == RefOutputTensors.Num())
-		{
-			for (int i = 0; i < OutputTensors.Num(); ++i)
-			{
-				bTestSuceeded &= VerifyTensorResult(
-					RefOutputTensors[i], RefOutputMemBuffers[i],
-					OutputTensors[i], OutputMemBuffers[i],
-					AbsoluteTolerance, RelativeTolerance);
-			}
-		}
-		else
-		{
-			UE_LOG(LogNNE, Error, TEXT("Expecting %d output tensor(s), got %d."), RefOutputTensors.Num(), OutputTensors.Num());
-			bTestSuceeded = false;
-		}
-		return bTestSuceeded;
-	}
-
 	bool RunTestInferenceAndCompareToRef(const FTests::FTestSetup& TestSetup, const FString& RuntimeName, const FNNIModelRaw& ONNXModel,
 		TArrayView<TArray<char>> RefOutputMemBuffers, TArrayView<NNECore::Internal::FTensor> RefOutputTensors)
 	{
@@ -650,7 +477,7 @@ namespace UE::NNEQA::Private
 		const float RelativeTolerance = TestSetup.GetRelativeToleranceForRuntime(RuntimeName);
 		bool bTestSuceeded = true;
 
-		if (!RunTestInferenceCPU(ONNXModel, TestSetup, RuntimeName, OutputTensors, OutputMemBuffers))
+		if (!RunTestInference(ONNXModel, TestSetup, RuntimeName, OutputTensors, OutputMemBuffers))
 		{
 			UE_LOG(LogNNE, Error, TEXT("Error running inference for engine %s."), *RuntimeName);
 			return false;
@@ -692,7 +519,7 @@ namespace UE::NNEQA::Private
 		TArray<NNECore::Internal::FTensor> RefOutputTensors;
 		bool bAllTestsSucceeded = true;
 
-		if (!RunTestInferenceCPU(ONNXModel, TestSetup, RefName, RefOutputTensors, RefOutputMemBuffers))
+		if (!RunTestInference(ONNXModel, TestSetup, RefName, RefOutputTensors, RefOutputMemBuffers))
 		{
 			UE_LOG(LogNNE, Error, TEXT("Error running reference inference with engine %s."), *RefName);
 			return false;
@@ -714,53 +541,6 @@ namespace UE::NNEQA::Private
 			UE_LOG(LogNNE, Error, TEXT("Expecting output from test setup are no matched by reference runtime."));
 		}
 
-		// Test against other runtime
-		//NNX Registry
-		for (auto Runtime : GetAllRuntimes())
-		{
-			const FString& RuntimeName = Runtime->GetRuntimeName();
-			if (RuntimeName == TEXT("NNXRuntimeCPU"))
-			{
-				continue;
-			}
-			if (!RuntimeFilter.IsEmpty() && !RuntimeFilter.Contains(RuntimeName))
-			{
-				continue;
-			}
-			FString TestResult;
-			if (TestSetup.AutomationExcludedRuntime.Contains(RuntimeName) ||
-				TestSetup.AutomationExcludedPlatformRuntimeCombination.Contains(TPair<FString, FString>(CurrentPlatform, RuntimeName)))
-			{
-				TestResult = TEXT("skipped (by config)");
-			}
-			else
-			{
-				bool bShouldRunVariadicTest = (ONNXModelVariadic.Format != ENNXInferenceFormat::Invalid);
-				bShouldRunVariadicTest &= !(RuntimeName == "NNXRuntimeDML");
-
-				bool bTestSuceeded = RunTestInferenceAndCompareToRef(TestSetup, Runtime, ONNXModel, RefOutputMemBuffers, RefOutputTensors);
-
-				if (bShouldRunVariadicTest)
-				{
-					if (!bTestSuceeded)
-					{
-						UE_LOG(LogNNE, Error, TEXT("Failed running static test."));
-					}
-					bool bVariadicTestSuceeded = RunTestInferenceAndCompareToRef(TestSetup, Runtime, ONNXModelVariadic, RefOutputMemBuffers, RefOutputTensors);
-					if (!bVariadicTestSuceeded)
-					{
-						bTestSuceeded = false;
-						UE_LOG(LogNNE, Error, TEXT("Failed running variadic test."));
-					}
-				}
-				TestResult = bTestSuceeded ? TEXT("SUCCESS") : TEXT("FAILED");
-				bAllTestsSucceeded &= bTestSuceeded;
-			}
-
-			UE_LOG(LogNNE, Display, TEXT("  %s tests: %s"), *RuntimeName, *TestResult);
-		}
-
-		//NNE Registry
 		for (auto Runtime : NNECore::GetAllRuntimes())
 		{
 			const FString& RuntimeName = Runtime->GetRuntimeName();
@@ -781,8 +561,8 @@ namespace UE::NNEQA::Private
 			else
 			{
 				bool bShouldRunVariadicTest = (ONNXModelVariadic.Format != ENNXInferenceFormat::Invalid);
-				bShouldRunVariadicTest &= !(RuntimeName == "NNEXRuntimeDML");
-
+				bShouldRunVariadicTest &= !(RuntimeName == "NNERuntimeDml");
+				
 				bool bTestSuceeded = RunTestInferenceAndCompareToRef(TestSetup, Runtime->GetRuntimeName(), ONNXModel, RefOutputMemBuffers, RefOutputTensors);
 
 				if (bShouldRunVariadicTest)
diff --git a/Engine/Plugins/Experimental/NNE/Source/NNEQA/Resources/NNEQATestDesc.json b/Engine/Plugins/Experimental/NNE/Source/NNEQA/Resources/NNEQATestDesc.json
--- a/Engine/Plugins/Experimental/NNE/Source/NNEQA/Resources/NNEQATestDesc.json
+++ b/Engine/Plugins/Experimental/NNE/Source/NNEQA/Resources/NNEQATestDesc.json
@@ -10,8 +10,8 @@
 					"tags": [ "SmokeTest" ],
 					"runtimes": [
 						{ "name": "NNERuntimeORTDml" },
-						{ "name": "NNXRuntimeHlsl",  "skip": true, "comment": "Missing BatchNormalisation operator." },
-						{ "name": "NNXRuntimeDml",  "skip": true, "comment": "Model not supported yet." }
+						{ "name": "NNERuntimeHlsl",  "skip": true, "comment": "Missing BatchNormalisation operator." },
+						{ "name": "NNERuntimeDml",  "skip": true, "comment": "Model not supported yet." }
 					],
 					"absolute_tolerance": 1e-5
 				},
@@ -20,8 +20,8 @@
 					"runtimes": [
 						{ "name": "NNERuntimeORTDml" },
 						{ "name": "NNERuntimeORTCpuEditor" },
-						{ "name": "NNXRuntimeHlsl" },
-						{ "name": "NNXRuntimeDml",  "skip": true, "comment": "Model not supported yet." }
+						{ "name": "NNERuntimeHlsl" },
+						{ "name": "NNERuntimeDml",  "skip": true, "comment": "Model not supported yet." }
 					],
 					"absolute_tolerance": 5.0e-2
 				}
@@ -36,8 +36,8 @@
 					"tags": [ "SmokeTest" ],
 					"runtimes": [
 						{ "name": "NNERuntimeORTDml" },
-						{ "name": "NNXRuntimeHlsl" },
-						{ "name": "NNXRuntimeDml" }
+						{ "name": "NNERuntimeHlsl" },
+						{ "name": "NNERuntimeDml" }
 					],
 					"absolute_tolerance": 1e-5
 				},
@@ -46,8 +46,8 @@
 					"tags": [ "SmokeTest" ],
 					"runtimes": [
 						{ "name": "NNERuntimeORTDml", "skip": true, "comment": "" },
-						{ "name": "NNXRuntimeHlsl",	"skip": true, "comment": "Model not supported yet."	},
-						{ "name": "NNXRuntimeDml",  "skip": true, "comment": "Model not supported yet." }
+						{ "name": "NNERuntimeHlsl",	"skip": true, "comment": "Model not supported yet."	},
+						{ "name": "NNERuntimeDml",  "skip": true, "comment": "Model not supported yet." }
 					]
 				},
 				{
@@ -55,16 +55,16 @@
 					"tags": [ "SmokeTest" ],
 					"runtimes": [
 						{ "name": "NNERuntimeORTDml", "absolute_tolerance": 1e-6 },
-						{ "name": "NNXRuntimeHlsl",  "skip": true, "comment": "Missing Reshape operator." },
-						{ "name": "NNXRuntimeDml",  "skip": true, "comment": "Model not supported yet." }
+						{ "name": "NNERuntimeHlsl",  "skip": true, "comment": "Missing Reshape operator." },
+						{ "name": "NNERuntimeDml",  "skip": true, "comment": "Model not supported yet." }
 					]
 				},
 				{
 					"target": "VertexDeltaModel",
 					"runtimes": [
 						{ "name": "NNERuntimeORTDml", "absolute_tolerance": 1e-4 },
-						{ "name": "NNXRuntimeHlsl", "absolute_tolerance": 1e-4 },
-						{ "name": "NNXRuntimeDml",  "skip": true, "comment": "Model not supported yet." }
+						{ "name": "NNERuntimeHlsl", "absolute_tolerance": 1e-4 },
+						{ "name": "NNERuntimeDml",  "skip": true, "comment": "Model not supported yet." }
 					]
 				}
 			]
@@ -76,21 +76,21 @@
 				{
 					"target": "OneLayerWithWeights",
 					"runtimes": [
-						{ "name": "NNXRuntimeDml",  "skip": true, "comment": "Model not supported yet." }
+						{ "name": "NNERuntimeDml",  "skip": true, "comment": "Model not supported yet." }
 					]
 				},
 				{
 					"target": "TwoLayersNoWeight",
 					"tags": [ "SmokeTest" ],
 					"runtimes": [
-						{ "name": "NNXRuntimeDml",  "skip": true, "comment": "Model not supported yet." }
+						{ "name": "NNERuntimeDml",  "skip": true, "comment": "Model not supported yet." }
 					]
 				},
 				{
 					"target": "MultiOutput",
 					"tags": [ "SmokeTest" ],
 					"runtimes": [
-						{ "name": "NNXRuntimeDml",  "skip": true, "comment": "Model not supported yet." }
+						{ "name": "NNERuntimeDml",  "skip": true, "comment": "Model not supported yet." }
 					]
 				}
 			]
@@ -102,14 +102,14 @@
 				{
 					"target": "OneLayerNoWeightVariableInputShape",
 					"runtimes": [
-						{ "name": "NNXRuntimeDML",  "skip": true, "comment": "RDG does not support variable input sized tensor" }
+						{ "name": "NNERuntimeDml",  "skip": true, "comment": "RDG does not support variable input sized tensor" }
 					]
 				},
 				{
 					"target": "TwoLayersNoWeightVariableInputShape",
 					"tags": [ "SmokeTest" ],
 					"runtimes": [
-						{ "name": "NNXRuntimeDML",  "skip": true, "comment": "RDG does not support variable input sized tensor" }
+						{ "name": "NNERuntimeDml",  "skip": true, "comment": "RDG does not support variable input sized tensor" }
 					]
 				}
 			]
@@ -123,76 +123,76 @@
 				{	"target": "Abs", "additional_datasets": [ "Elementwise.Slow.OneInput" ]},
 				{	"target": "Acos", "tags": [ "SmokeTest" ], "runtimes": [
 						{ "name": "NNERuntimeORTDml", "absolute_tolerance":  1e-4},
-						{ "name": "NNXRuntimeDml",    "absolute_tolerance":  1e-4},
-						{ "name": "NNXRuntimeHlsl",   "absolute_tolerance":  2e-4}
+						{ "name": "NNERuntimeDml",    "absolute_tolerance":  1e-4},
+						{ "name": "NNERuntimeHlsl",   "absolute_tolerance":  2e-4}
 					]},
 				{	"target": "Acosh" },
 				{	"target": "Asin", "runtimes": [
 						{ "name": "NNERuntimeORTDml", "absolute_tolerance":  1e-4},
-						{ "name": "NNXRuntimeDml",    "absolute_tolerance":  1e-4},
-						{ "name": "NNXRuntimeHlsl",   "absolute_tolerance":  2e-4}
+						{ "name": "NNERuntimeDml",    "absolute_tolerance":  1e-4},
+						{ "name": "NNERuntimeHlsl",   "absolute_tolerance":  2e-4}
 					]},
 				{	"target": "Asinh", "runtimes": [
 						{ "name": "NNERuntimeORTDml", "absolute_tolerance":  1e-4},
-						{ "name": "NNXRuntimeDml",    "absolute_tolerance":  1e-4},
-						{ "name": "NNXRuntimeHlsl",   "absolute_tolerance":  2e-4}
+						{ "name": "NNERuntimeDml",    "absolute_tolerance":  1e-4},
+						{ "name": "NNERuntimeHlsl",   "absolute_tolerance":  2e-4}
 					]},
 				{	"target": "Atan", "runtimes": [
 						{ "name": "NNERuntimeORTDml", "absolute_tolerance":  1e-4},
-						{ "name": "NNXRuntimeDml",    "absolute_tolerance":  1e-4},
-						{ "name": "NNXRuntimeHlsl",   "absolute_tolerance":  1.5e-3}
+						{ "name": "NNERuntimeDml",    "absolute_tolerance":  1e-4},
+						{ "name": "NNERuntimeHlsl",   "absolute_tolerance":  1.5e-3}
 					]},
 				{	"target": "Atanh" },
 				{	"target": "BitShift", "skip": true, "comment": "TODO need attributes" },
 				{	"target": "Cast",     "skip": true, "comment": "TODO need attributes" },
 				{	"target": "Ceil"},
 				{	"target": "Clip", "additional_datasets": [ "Elementwise.ClipTestWithExtraInput" ], "runtimes": [
-						{ "name": "NNXRuntimeDml", "skip": true, "comment": "RDG need scalar tensor support"},
-						{ "name": "NNXRuntimeHlsl", "skip": true, "comment": "RDG need scalar tensor support"}						
+						{ "name": "NNERuntimeDml", "skip": true, "comment": "RDG need scalar tensor support"},
+						{ "name": "NNERuntimeHlsl", "skip": true, "comment": "RDG need scalar tensor support"}						
 					]},
 				{	"target": "Cos" },
 				{	"target": "Cosh" },
 				{	"target": "Elu", "comment": "TODO test with not-default attribute", "runtimes": [
 						{ "name": "NNERuntimeORTDml", "absolute_tolerance":  1e-4},
-						{ "name": "NNXRuntimeDml",    "absolute_tolerance":  1e-4},
-						{ "name": "NNXRuntimeHlsl",   "absolute_tolerance":  1e-4}
+						{ "name": "NNERuntimeDml",    "absolute_tolerance":  1e-4},
+						{ "name": "NNERuntimeHlsl",   "absolute_tolerance":  1e-4}
 					]},
 				{	"target": "Erf", "runtimes": [
 						{ "name": "NNERuntimeORTDml", "absolute_tolerance":  1e-4},
-						{ "name": "NNXRuntimeDml",    "absolute_tolerance":  1e-4},
-						{ "name": "NNXRuntimeHlsl",   "absolute_tolerance":  5e-4, "comment": "NNXRuntimeHLSL Erf implementation is less precise than the DML one"}
+						{ "name": "NNERuntimeDml",    "absolute_tolerance":  1e-4},
+						{ "name": "NNERuntimeHlsl",   "absolute_tolerance":  5e-4, "comment": "NNERuntimeHlsl Erf implementation is less precise than the DML one"}
 					]},
 				{	"target": "Exp" },
 				{	"target": "Floor" },
 				{	"target": "IsInf", "skip": true, "comment": "TODO add flexibility in test data setup + bool tensors" },
 				{	"target": "IsNan", "skip": true, "comment": "TODO add flexibility in test data setup + bool tensors" },
 				{	"target": "HardSigmoid",         "comment": "TODO test with not-default attribute" },
-				{	"target": "HardSwish", "runtimes": [ { "name": "NNXRuntimeDml", "skip": true, "comment": "Not implemented atm on NNXRuntimeDML."} ]},
+				{	"target": "HardSwish", "runtimes": [ { "name": "NNERuntimeDml", "skip": true, "comment": "Not implemented atm on NNERuntimeDml."} ]},
 				{	"target": "LeakyRelu", "additional_datasets": [ "Elementwise.LeakyRelu.AttributeTest" ] },
 				{	"target": "Log", "runtimes": [
 						{ "name": "NNERuntimeORTDml"},
-						{ "name": "NNXRuntimeDml"},
-						{ "name": "NNXRuntimeHlsl"}
+						{ "name": "NNERuntimeDml"},
+						{ "name": "NNERuntimeHlsl"}
 					], "absolute_tolerance": 1e-6 },
 				{	"target": "Neg" },
 				{	"target": "Not", "input_type": "boolean", "output_type": "boolean", "runtimes": [
-						{ "name": "NNXRuntimeDml", "skip": true, "comment": "RDG need bool tensor support"},
-						{ "name": "NNXRuntimeHlsl", "skip": true, "comment": "RDG need bool tensor support"}						
+						{ "name": "NNERuntimeDml", "skip": true, "comment": "RDG need bool tensor support"},
+						{ "name": "NNERuntimeHlsl", "skip": true, "comment": "RDG need bool tensor support"}						
 					]},
 				{	"target": "Reciprocal"},
 				{	"target": "Relu"},
 				{	"target": "Round"},
 				{	"target": "Selu", "comment": "TODO test with not-default attribute", "runtimes": [
 						{ "name": "NNERuntimeORTDml", "absolute_tolerance":  1e-4},
-						{ "name": "NNXRuntimeDml",    "absolute_tolerance":  1e-4},
-						{ "name": "NNXRuntimeHlsl",   "absolute_tolerance":  1e-4}
+						{ "name": "NNERuntimeDml",    "absolute_tolerance":  1e-4},
+						{ "name": "NNERuntimeHlsl",   "absolute_tolerance":  1e-4}
 					]},
 				{	"target": "Sigmoid"},
 				{	"target": "Sign" },
 				{	"target": "Sin", "runtimes": [
 						{ "name": "NNERuntimeORTDml", "absolute_tolerance":  1e-4},
-						{ "name": "NNXRuntimeDml",    "absolute_tolerance":  1e-4},
-						{ "name": "NNXRuntimeHlsl",   "absolute_tolerance":  1e-4}
+						{ "name": "NNERuntimeDml",    "absolute_tolerance":  1e-4},
+						{ "name": "NNERuntimeHlsl",   "absolute_tolerance":  1e-4}
 					]},
 				{	"target": "Sinh" },
 				{	"target": "Softplus"},
@@ -200,13 +200,13 @@
 				{	"target": "Sqrt"},
 				{	"target": "Tan", "runtimes": [
 						{ "name": "NNERuntimeORTDml", "absolute_tolerance":  1e-4},
-						{ "name": "NNXRuntimeDml",    "absolute_tolerance":  1e-4},
-						{ "name": "NNXRuntimeHlsl",   "absolute_tolerance":  1e-4}
+						{ "name": "NNERuntimeDml",    "absolute_tolerance":  1e-4},
+						{ "name": "NNERuntimeHlsl",   "absolute_tolerance":  1e-4}
 					]},
 				{	"target": "Tanh", "runtimes": [
 						{ "name": "NNERuntimeORTDml", "absolute_tolerance":  1e-4},
-						{ "name": "NNXRuntimeDml",    "absolute_tolerance":  1e-4},
-						{ "name": "NNXRuntimeHlsl",   "absolute_tolerance":  1e-4}
+						{ "name": "NNERuntimeDml",    "absolute_tolerance":  1e-4},
+						{ "name": "NNERuntimeHlsl",   "absolute_tolerance":  1e-4}
 					]}
 			]
 		},
@@ -215,35 +215,35 @@
 			"targets": [
 				{	"target": "Add", "additional_datasets": [ "Elementwise.Slow.TwoInputs" ]},
 				{	"target": "And", "input_type": "boolean", "output_type": "boolean", "runtimes": [
-						{ "name": "NNXRuntimeDml", "skip": true, "comment": "RDG need bool tensor support"},
-						{ "name": "NNXRuntimeHlsl", "skip": true, "comment": "RDG need bool tensor support"}						
+						{ "name": "NNERuntimeDml", "skip": true, "comment": "RDG need bool tensor support"},
+						{ "name": "NNERuntimeHlsl", "skip": true, "comment": "RDG need bool tensor support"}						
 					]},
 				{	"target": "Div", "tags": [ "SmokeTest" ]},
 				{	"target": "Equal", "output_type": "boolean", "runtimes": [
-						{ "name": "NNXRuntimeDml", "skip": true, "comment": "RDG need bool tensor support"},
-						{ "name": "NNXRuntimeHlsl", "skip": true, "comment": "RDG need bool tensor support"}						
+						{ "name": "NNERuntimeDml", "skip": true, "comment": "RDG need bool tensor support"},
+						{ "name": "NNERuntimeHlsl", "skip": true, "comment": "RDG need bool tensor support"}						
 					]},
 				{	"target": "Greater", "output_type": "boolean", "runtimes": [
-						{ "name": "NNXRuntimeDml", "skip": true, "comment": "RDG need bool tensor support"},
-						{ "name": "NNXRuntimeHlsl", "skip": true, "comment": "RDG need bool tensor support"}						
+						{ "name": "NNERuntimeDml", "skip": true, "comment": "RDG need bool tensor support"},
+						{ "name": "NNERuntimeHlsl", "skip": true, "comment": "RDG need bool tensor support"}						
 					]},
 				{	"target": "GreaterOrEqual", "output_type": "boolean", "runtimes": [
-						{ "name": "NNXRuntimeDml", "skip": true, "comment": "RDG need bool tensor support"},
-						{ "name": "NNXRuntimeHlsl", "skip": true, "comment": "RDG need bool tensor support"}						
+						{ "name": "NNERuntimeDml", "skip": true, "comment": "RDG need bool tensor support"},
+						{ "name": "NNERuntimeHlsl", "skip": true, "comment": "RDG need bool tensor support"}						
 					]},
 				{	"target": "Less", "output_type": "boolean", "runtimes": [
-						{ "name": "NNXRuntimeDml", "skip": true, "comment": "RDG need bool tensor support"},
-						{ "name": "NNXRuntimeHlsl", "skip": true, "comment": "RDG need bool tensor support"}						
+						{ "name": "NNERuntimeDml", "skip": true, "comment": "RDG need bool tensor support"},
+						{ "name": "NNERuntimeHlsl", "skip": true, "comment": "RDG need bool tensor support"}						
 					]},
 				{	"target": "LessOrEqual", "output_type": "boolean", "runtimes": [
-						{ "name": "NNXRuntimeDml", "skip": true, "comment": "RDG need bool tensor support"},
-						{ "name": "NNXRuntimeHlsl", "skip": true, "comment": "RDG need bool tensor support"}						
+						{ "name": "NNERuntimeDml", "skip": true, "comment": "RDG need bool tensor support"},
+						{ "name": "NNERuntimeHlsl", "skip": true, "comment": "RDG need bool tensor support"}						
 					]},
 				{	"target": "Mod", "skip": true, "comment": "TODO need to set fmod attribute to 1 for floating point tensors"},
 				{	"target": "Mul"},
 				{	"target": "Or", "input_type": "boolean", "output_type": "boolean", "runtimes": [
-						{ "name": "NNXRuntimeDml", "skip": true, "comment": "RDG need bool tensor support"},
-						{ "name": "NNXRuntimeHlsl", "skip": true, "comment": "RDG need bool tensor support"}						
+						{ "name": "NNERuntimeDml", "skip": true, "comment": "RDG need bool tensor support"},
+						{ "name": "NNERuntimeHlsl", "skip": true, "comment": "RDG need bool tensor support"}						
 					]},
 				{	"target": "PRelu", "removed_datasets": [ 
 						"Elementwise.Binary.BidirectionalBroadcast", 
@@ -252,15 +252,15 @@
 				{	"target": "Pow", "absolute_tolerance": 1e-4},
 				{	"target": "Sub"},
 				{	"target": "Xor", "input_type": "boolean", "output_type": "boolean", "runtimes": [
-						{ "name": "NNXRuntimeDml", "skip": true, "comment": "RDG need bool tensor support"},
-						{ "name": "NNXRuntimeHlsl", "skip": true, "comment": "RDG need bool tensor support"}						
+						{ "name": "NNERuntimeDml", "skip": true, "comment": "RDG need bool tensor support"},
+						{ "name": "NNERuntimeHlsl", "skip": true, "comment": "RDG need bool tensor support"}						
 					]}
 			]
 		},
 		{
 			"category": "Elementwise.Variadic",
 			"additional_datasets": [ "Elementwise.Binary" ],
-			"runtimes": [ { "name": "NNXRuntimeDml", "skip": true} ],
+			"runtimes": [ { "name": "NNERuntimeDml", "skip": true} ],
 			"targets": [
 				{	"target": "Max", "additional_datasets": [ "Elementwise.Slow.TwoInputs" ]},
 				{	"target": "Mean", "tags": [ "SmokeTest" ]},
@@ -277,9 +277,9 @@
 		{
 			"category": "Convolution",
 			"runtimes": [
-				{ "name": "NNXRuntimeDml", "skip": true	},
+				{ "name": "NNERuntimeDml", "skip": true	},
 				{ "name": "NNERuntimeORTDml" },
-				{ "name": "NNXRuntimeHlsl" }
+				{ "name": "NNERuntimeHlsl" }
 			],
 			"targets": [
 				{ "target": "Conv", "absolute_tolerance": 1e-6 }
@@ -288,9 +288,9 @@
 		{
 			"category": "ConvTranspose",
 			"runtimes": [
-				{ "name": "NNXRuntimeDml", "skip": true	},
+				{ "name": "NNERuntimeDml", "skip": true	},
 				{ "name": "NNERuntimeORTDml" },
-				{ "name": "NNXRuntimeHlsl" }
+				{ "name": "NNERuntimeHlsl" }
 			],
 			"targets": [
 				{ "target": "ConvTranspose", "absolute_tolerance": 1e-6 }
@@ -298,7 +298,7 @@
 		},
 		{
 			"category": "MatMul",
-			"runtimes": [ { "name": "NNXRuntimeDml", "skip": true } ],
+			"runtimes": [ { "name": "NNERuntimeDml", "skip": true } ],
 			"targets": [
 				{ "target": "MatMul" }
 			]
@@ -306,7 +306,7 @@
 		{
 			"category": "PadOp",
 			"runtimes": [
-				{ "name": "NNXRuntimeDml", "skip": true	}
+				{ "name": "NNERuntimeDml", "skip": true	}
 			],
 			"targets": [
 				{ "target": "Pad" }
@@ -315,7 +315,7 @@
 		{
 			"category": "InstanceNormalization",
 			"runtimes": [
-				{ "name": "NNXRuntimeDml", "skip": true }
+				{ "name": "NNERuntimeDml", "skip": true }
 			],
 			"targets": [
 				{ "target": "InstanceNormalization", "absolute_tolerance": 1e-5 }
@@ -324,7 +324,7 @@
 		{
 			"category": "Upsample",
 			"runtimes": [
-				{ "name": "NNXRuntimeDml", "skip": true	}
+				{ "name": "NNERuntimeDml", "skip": true	}
 			],
 			"targets": [
 				{ "target": "Upsample" }
@@ -353,8 +353,8 @@
 			"datasets": [
 				{ "inputs": [ { "shape": [ 1,0,3 ] } ],
 					"runtimes": [
-						{ "name": "NNXRuntimeHLSL", "skip": true, "comment": "RDG does not support 0 sized tensor" },
-						{ "name": "NNXRuntimeDML",  "skip": true, "comment": "RDG does not support 0 sized tensor" }
+						{ "name": "NNERuntimeHlsl", "skip": true, "comment": "RDG does not support 0 sized tensor" },
+						{ "name": "NNERuntimeDml",  "skip": true, "comment": "RDG does not support 0 sized tensor" }
 					]}
 			]
 		},
@@ -403,8 +403,8 @@
 			"datasets": [
 				{ "inputs": [ { "shape": [ 1,0 ] }, { "shape": [ 1,1 ] } ],
 					"runtimes": [
-						{ "name": "NNXRuntimeHLSL", "skip": true, "comment": "RDG does not support 0 sized tensor" },
-						{ "name": "NNXRuntimeDML",  "skip": true, "comment": "RDG does not support 0 sized tensor" }
+						{ "name": "NNERuntimeHlsl", "skip": true, "comment": "RDG does not support 0 sized tensor" },
+						{ "name": "NNERuntimeDml",  "skip": true, "comment": "RDG does not support 0 sized tensor" }
 					]}
 			]
 		},
@@ -562,7 +562,7 @@
 			"datasets": [ { "inputs": [ { "shape": [ 1, 1, 3, 3 ] }, { "shape": [ 1, 2, 3, 3 ] } ], "outputs": [ { "shape": [ 1, 2, 6, 6 ] } ] } ],
 			"runtimes": [
 				{ "name": "NNERuntimeORTDml", "skip": true, "comment": "TODO INVESTIGATE Returns 'The parameter is incorrect'" },
-				{ "name": "NNXRuntimeHlsl", "skip": true, "comment": "Fails until ORT 1.13, because they implemented 'auto_pad' differently (they flipped 'SAME_LOWER' and 'SAME_UPPER')" }
+				{ "name": "NNERuntimeHlsl", "skip": true, "comment": "Fails until ORT 1.13, because they implemented 'auto_pad' differently (they flipped 'SAME_LOWER' and 'SAME_UPPER')" }
 			]
 		},
 		{
@@ -584,7 +584,7 @@
 			"name": "MatMul.Rank5",
 			"runtimes": [	
 				{ "name": "NNERuntimeORTDml", "skip": true, "comment": "DML metacommand only support 2 to 4 for inputs rank" },
-				{ "name": "NNXRuntimeHlsl"}
+				{ "name": "NNERuntimeHlsl"}
 			],
 			"datasets": [
 				{ "inputs": [ { "shape": [ 2, 2, 3, 5, 4 ] }, { "shape": [ 2, 1, 3, 4, 3 ] } ], "outputs": [ { "shape": [ 2, 2, 3, 5, 3 ] } ]},
@@ -628,8 +628,8 @@
 		{
 			"name": "Upsample.ScaleFromInput",
 			"runtimes": [	
-				{ "name": "NNXRuntimeDml", "skip": true, "comment": "Scale should be a constant tensor for RDG" },
-				{ "name": "NNXRuntimeHlsl", "skip": true, "comment": "Scale should be a constant tensor for RDG" }
+				{ "name": "NNERuntimeDml", "skip": true, "comment": "Scale should be a constant tensor for RDG" },
+				{ "name": "NNERuntimeHlsl", "skip": true, "comment": "Scale should be a constant tensor for RDG" }
 			],
 			"datasets": [
 				{ "inputs":  [ { "shape": [ 1 ] },                            { "shape": [ 1 ], "source": [ "3"    ],                  "type": "float" } ], "outputs": [ { "shape": [ 3 ]} ] },
@@ -661,7 +661,7 @@
 				{
 					"inputs": [ { "shape": [ 2, 2, 4, 1, 3 ] }, { "shape": [ 2 ] }, { "shape": [ 2 ] } ], "outputs": [ { "shape": [ 2, 2, 4, 1, 3 ] } ],
 					"runtimes": [
-						{ "name": "NNXRuntimeHLSL" },
+						{ "name": "NNERuntimeHlsl" },
 						{ "name": "NNERuntimeORTCpuEditor" },
 						{ "name": "NNERuntimeORTDml", "skip": true, "comment": "Not supported" }
 					]
@@ -669,7 +669,7 @@
 				{
 					"inputs": [ { "shape": [ 2, 2, 4, 4 ] }, { "shape": [ 2 ] }, { "shape": [ 2 ] } ], "outputs": [ { "shape": [ 2, 2, 4, 4 ] } ],
 					"runtimes": [
-						{ "name": "NNXRuntimeHLSL" },
+						{ "name": "NNERuntimeHlsl" },
 						{ "name": "NNERuntimeORTCpuEditor" },
 						{ "name": "NNERuntimeORTDml" }
 					]
@@ -677,7 +677,7 @@
 				{
 					"inputs": [ { "shape": [ 2, 2, 8, 8 ] }, { "shape": [ 2 ] }, { "shape": [ 2 ] } ], "outputs": [ { "shape": [ 2, 2, 8, 8 ] } ],
 					"runtimes": [
-						{ "name": "NNXRuntimeHLSL" },
+						{ "name": "NNERuntimeHlsl" },
 						{ "name": "NNERuntimeORTCpuEditor" },
 						{ "name": "NNERuntimeORTDml" }
 					]
@@ -685,7 +685,7 @@
 				{
 					"inputs": [ { "shape": [ 2, 128, 8, 8 ] }, { "shape": [ 128 ] }, { "shape": [ 128 ] } ], "outputs": [ { "shape": [ 2, 128, 8, 8 ] } ],
 					"runtimes": [
-						{ "name": "NNXRuntimeHLSL" },
+						{ "name": "NNERuntimeHlsl" },
 						{ "name": "NNERuntimeORTCpuEditor" },
 						{ "name": "NNERuntimeORTDml" }
 					]
@@ -693,7 +693,7 @@
 				{
 					"inputs": [ { "shape": [ 128, 2, 8, 8 ] }, { "shape": [ 2 ] }, { "shape": [ 2 ] } ], "outputs": [ { "shape": [ 128, 2, 8, 8 ] } ],
 					"runtimes": [
-						{ "name": "NNXRuntimeHLSL" },
+						{ "name": "NNERuntimeHlsl" },
 						{ "name": "NNERuntimeORTCpuEditor" },
 						{ "name": "NNERuntimeORTDml" }
 					]
@@ -701,7 +701,7 @@
 				{
 					"inputs": [ { "shape": [ 64, 64, 8, 8 ] }, { "shape": [ 64 ] }, { "shape": [ 64 ] } ], "outputs": [ { "shape": [ 64, 64, 8, 8 ] } ],
 					"runtimes": [
-						{ "name": "NNXRuntimeHLSL" },
+						{ "name": "NNERuntimeHlsl" },
 						{ "name": "NNERuntimeORTCpuEditor" },
 						{ "name": "NNERuntimeORTDml" }
 					]
@@ -709,7 +709,7 @@
 				{
 					"inputs": [ { "shape": [ 2, 2, 45, 45 ] }, { "shape": [ 2 ] }, { "shape": [ 2 ] } ], "outputs": [ { "shape": [ 2, 2, 45, 45 ] } ],
 					"runtimes": [
-						{ "name": "NNXRuntimeHLSL" },
+						{ "name": "NNERuntimeHlsl" },
 						{ "name": "NNERuntimeORTCpuEditor" },
 						{ "name": "NNERuntimeORTDml" }
 					]
@@ -717,7 +717,7 @@
 				{
 					"inputs": [ { "shape": [ 1, 128, 56, 56 ] }, { "shape": [ 128 ] }, { "shape": [ 128 ] } ], "outputs": [ { "shape": [ 1, 128, 56, 56 ] } ],
 					"runtimes": [
-						{ "name": "NNXRuntimeHLSL" },
+						{ "name": "NNERuntimeHlsl" },
 						{ "name": "NNERuntimeORTCpuEditor" },
 						{ "name": "NNERuntimeORTDml" }
 					],
diff --git a/Engine/Plugins/Experimental/NNE/Source/NNERuntimeRDG/Private/Hlsl/NNERuntimeRDGConv.cpp b/Engine/Plugins/Experimental/NNE/Source/NNERuntimeRDG/Private/Hlsl/NNERuntimeRDGConv.cpp
--- a/Engine/Plugins/Experimental/NNE/Source/NNERuntimeRDG/Private/Hlsl/NNERuntimeRDGConv.cpp
+++ b/Engine/Plugins/Experimental/NNE/Source/NNERuntimeRDG/Private/Hlsl/NNERuntimeRDGConv.cpp
@@ -74,17 +74,17 @@ namespace UE::NNERuntimeRDG::Private::Hlsl
 			
 			if (Input.GetShape().Rank() < 2)
 			{
-				UE_LOG(LogNNX, Warning, TEXT("Conv first input should be at least of rank 2"));
+				UE_LOG(LogNNE, Warning, TEXT("Conv first input should be at least of rank 2"));
 				return false;
 			}
 			if (Weights.GetShape().Rank() != Input.GetShape().Rank())
 			{
-				UE_LOG(LogNNX, Warning, TEXT("Conv first and second inputs should be of same ranks"));
+				UE_LOG(LogNNE, Warning, TEXT("Conv first and second inputs should be of same ranks"));
 				return false;
 			}
 			if (Output.GetShape().Rank() != Input.GetShape().Rank())
 			{
-				UE_LOG(LogNNX, Warning, TEXT("Conv first and output should be of same ranks"));
+				UE_LOG(LogNNE, Warning, TEXT("Conv first and output should be of same ranks"));
 				return false;
 			}
 
diff --git a/Engine/Plugins/Experimental/NNE/Source/NNERuntimeRDG/Private/Hlsl/NNERuntimeRDGConvTranspose.cpp b/Engine/Plugins/Experimental/NNE/Source/NNERuntimeRDG/Private/Hlsl/NNERuntimeRDGConvTranspose.cpp
--- a/Engine/Plugins/Experimental/NNE/Source/NNERuntimeRDG/Private/Hlsl/NNERuntimeRDGConvTranspose.cpp
+++ b/Engine/Plugins/Experimental/NNE/Source/NNERuntimeRDG/Private/Hlsl/NNERuntimeRDGConvTranspose.cpp
@@ -74,17 +74,17 @@ namespace UE::NNERuntimeRDG::Private::Hlsl
 
 			if (Input.GetShape().Rank() < 2)
 			{
-				UE_LOG(LogNNX, Warning, TEXT("ConvTranspose first input should be at least of rank 2"));
+				UE_LOG(LogNNE, Warning, TEXT("ConvTranspose first input should be at least of rank 2"));
 				return false;
 			}
 			if (Weights.GetShape().Rank() != Input.GetShape().Rank())
 			{
-				UE_LOG(LogNNX, Warning, TEXT("ConvTranspose first and second inputs should be of same ranks"));
+				UE_LOG(LogNNE, Warning, TEXT("ConvTranspose first and second inputs should be of same ranks"));
 				return false;
 			}
 			if (Output.GetShape().Rank() != Input.GetShape().Rank())
 			{
-				UE_LOG(LogNNX, Warning, TEXT("ConvTranspose first input and output should be of same ranks"));
+				UE_LOG(LogNNE, Warning, TEXT("ConvTranspose first input and output should be of same ranks"));
 				return false;
 			}
 
diff --git a/Engine/Plugins/Experimental/NNE/Source/NNERuntimeRDG/Private/Hlsl/NNERuntimeRDGElementWiseBinary.cpp b/Engine/Plugins/Experimental/NNE/Source/NNERuntimeRDG/Private/Hlsl/NNERuntimeRDGElementWiseBinary.cpp
--- a/Engine/Plugins/Experimental/NNE/Source/NNERuntimeRDG/Private/Hlsl/NNERuntimeRDGElementWiseBinary.cpp
+++ b/Engine/Plugins/Experimental/NNE/Source/NNERuntimeRDG/Private/Hlsl/NNERuntimeRDGElementWiseBinary.cpp
@@ -46,7 +46,7 @@ namespace UE::NNERuntimeRDG::Private::Hlsl
 				int32 RHSValue = RHSIndex >= 0 ? RHSInput.GetData()[RHSIndex] : 1;
 				if (LHSValue != RHSValue && LHSValue != 1 && RHSValue != 1)
 				{
-					UE_LOG(LogNNX, Warning, TEXT("Error while computing shape for element wise binary op, input shapes are not compatible"));
+					UE_LOG(LogNNE, Warning, TEXT("Error while computing shape for element wise binary op, input shapes are not compatible"));
 					return -1;
 				}
 				int32 OutputValue = FMath::Max(LHSValue, RHSValue);
diff --git a/Engine/Plugins/Experimental/NNE/Source/NNERuntimeRDG/Private/Hlsl/NNERuntimeRDGElementWiseVariadic.cpp b/Engine/Plugins/Experimental/NNE/Source/NNERuntimeRDG/Private/Hlsl/NNERuntimeRDGElementWiseVariadic.cpp
--- a/Engine/Plugins/Experimental/NNE/Source/NNERuntimeRDG/Private/Hlsl/NNERuntimeRDGElementWiseVariadic.cpp
+++ b/Engine/Plugins/Experimental/NNE/Source/NNERuntimeRDG/Private/Hlsl/NNERuntimeRDGElementWiseVariadic.cpp
@@ -119,7 +119,7 @@ namespace UE::NNERuntimeRDG::Private::Hlsl
 					int32 InputValue = InputIndex >= 0 ? InputTensors[InputIdx]->GetShape().GetData()[InputIndex] : 1;
 					if (InputValue != OutputValue && InputValue != 1 && OutputValue != 1)
 					{
-						UE_LOG(LogNNX, Warning, TEXT("Error while computing shape for element wise variadic op, input shapes are not compatible"));
+						UE_LOG(LogNNE, Warning, TEXT("Error while computing shape for element wise variadic op, input shapes are not compatible"));
 						return -1;
 					}
 					OutputValue = FMath::Max(InputValue, OutputValue);
@@ -188,14 +188,14 @@ namespace UE::NNERuntimeRDG::Private::Hlsl
 
 		if (InputTypes.Num() == 0)
 		{
-			UE_LOG(LogNNX, Error, TEXT("Element-wise variadic operator requires at least 1 input"));
+			UE_LOG(LogNNE, Error, TEXT("Element-wise variadic operator requires at least 1 input"));
 			bIsValid = false;
 		}
 		for (int32 i = 0; i < InputTypes.Num(); ++i)
 		{
 			if (InputTypes[i] != ENNETensorDataType::Float)
 			{
-				UE_LOG(LogNNX, Warning, TEXT("Element-wise variadic operator input '%d' of type '%d' is not supported, should be float at the moment."), i, InputTypes[i]);
+				UE_LOG(LogNNE, Warning, TEXT("Element-wise variadic operator input '%d' of type '%d' is not supported, should be float at the moment."), i, InputTypes[i]);
 				bIsValid = false;
 			}
 		}
diff --git a/Engine/Plugins/Experimental/NNE/Source/NNERuntimeRDG/Private/Hlsl/NNERuntimeRDGGather.cpp b/Engine/Plugins/Experimental/NNE/Source/NNERuntimeRDG/Private/Hlsl/NNERuntimeRDGGather.cpp
--- a/Engine/Plugins/Experimental/NNE/Source/NNERuntimeRDG/Private/Hlsl/NNERuntimeRDGGather.cpp
+++ b/Engine/Plugins/Experimental/NNE/Source/NNERuntimeRDG/Private/Hlsl/NNERuntimeRDGGather.cpp
@@ -32,7 +32,7 @@ namespace UE::NNERuntimeRDG::Private::Hlsl
 		{
 			check(InputTensors.Num() == 2)
 			check(OutputTensors.Num() == 1)
-			UE_LOG(LogNNX, Warning, TEXT("Gather shape inference is not implemented at the moment"));
+			UE_LOG(LogNNE, Warning, TEXT("Gather shape inference is not implemented at the moment"));
 			return -1;
 		};
 		
@@ -49,34 +49,34 @@ namespace UE::NNERuntimeRDG::Private::Hlsl
 
 			if (Output.GetShape().Rank() <= MaxNumDimensions)
 			{
-				UE_LOG(LogNNX, Warning, TEXT("Gather first input should be of rank %d or less"), MaxNumDimensions);
+				UE_LOG(LogNNE, Warning, TEXT("Gather first input should be of rank %d or less"), MaxNumDimensions);
 				return false;
 			}
 			if (Data.GetShape().Rank() == 0)
 			{
-				UE_LOG(LogNNX, Warning, TEXT("Gather first input should be at least of rank 1"));
+				UE_LOG(LogNNE, Warning, TEXT("Gather first input should be at least of rank 1"));
 				return false;
 			}
 			if (Indices.GetShape().Rank() == 0)
 			{
-				UE_LOG(LogNNX, Warning, TEXT("Gather second input should be at least of rank 1"));
+				UE_LOG(LogNNE, Warning, TEXT("Gather second input should be at least of rank 1"));
 				return false;
 			}
 			if ((Data.GetShape().Rank() + Indices.GetShape().Rank() - 1) > MaxNumDimensions)
 			{
-				UE_LOG(LogNNX, Warning, TEXT("Gather sum of input 0 and 1 ranks -1 should be less than %d"), MaxNumDimensions);
+				UE_LOG(LogNNE, Warning, TEXT("Gather sum of input 0 and 1 ranks -1 should be less than %d"), MaxNumDimensions);
 				return false;
 			}
 
 			Axis = Attributes.GetValueOrDefault(TEXT("axis"), Axis);
 			if (Axis >= Data.GetShape().Rank())
 			{
-				UE_LOG(LogNNX, Warning, TEXT("Gather Axis attribute should be inferior to first input rank"));
+				UE_LOG(LogNNE, Warning, TEXT("Gather Axis attribute should be inferior to first input rank"));
 				return false;
 			}
 			if (Axis < -Data.GetShape().Rank())
 {
-				UE_LOG(LogNNX, Warning, TEXT("Gather Axis attribute should be superior or equal to minus the first input rank"));
+				UE_LOG(LogNNE, Warning, TEXT("Gather Axis attribute should be superior or equal to minus the first input rank"));
 				return false;
 			}
 			Axis = Axis >= 0 ? Axis : Data.GetShape().Rank() + Axis;
diff --git a/Engine/Plugins/Experimental/NNE/Source/NNERuntimeRDG/Private/Hlsl/NNERuntimeRDGGemm.cpp b/Engine/Plugins/Experimental/NNE/Source/NNERuntimeRDG/Private/Hlsl/NNERuntimeRDGGemm.cpp
--- a/Engine/Plugins/Experimental/NNE/Source/NNERuntimeRDG/Private/Hlsl/NNERuntimeRDGGemm.cpp
+++ b/Engine/Plugins/Experimental/NNE/Source/NNERuntimeRDG/Private/Hlsl/NNERuntimeRDGGemm.cpp
@@ -68,12 +68,12 @@ namespace UE::NNERuntimeRDG::Private::Hlsl
 
 			if (InputA.GetShape().Rank() != 2)
 			{
-				UE_LOG(LogNNX, Warning, TEXT("Gemm first input should be of rank 2"));
+				UE_LOG(LogNNE, Warning, TEXT("Gemm first input should be of rank 2"));
 				return false;
 			}
 			if (InputB.GetShape().Rank() != 2)
 			{
-				UE_LOG(LogNNX, Warning, TEXT("Gemm second input should be of rank 2"));
+				UE_LOG(LogNNE, Warning, TEXT("Gemm second input should be of rank 2"));
 				return false;
 			}
 			if (InputTensorDescs.Num() == 3)
@@ -81,12 +81,12 @@ namespace UE::NNERuntimeRDG::Private::Hlsl
 				const NNECore::FTensorDesc& InputC = InputTensorDescs[2];
 				if (InputC.GetShape().Rank() > 2)
 				{
-					UE_LOG(LogNNX, Warning, TEXT("Gemm third input should be of rank 2 or less"));
+					UE_LOG(LogNNE, Warning, TEXT("Gemm third input should be of rank 2 or less"));
 					return false;
 				}
 				if (InputC.GetShape().Rank() == 1 && InputC.GetShape().GetData()[0] == 1)
 				{
-					UE_LOG(LogNNX, Warning, TEXT("Gemm third input as scalar not supported"));
+					UE_LOG(LogNNE, Warning, TEXT("Gemm third input as scalar not supported"));
 					return false;
 				}
 			}
diff --git a/Engine/Plugins/Experimental/NNE/Source/NNERuntimeRDG/Private/Hlsl/NNERuntimeRDGInstanceNormalization.cpp b/Engine/Plugins/Experimental/NNE/Source/NNERuntimeRDG/Private/Hlsl/NNERuntimeRDGInstanceNormalization.cpp
--- a/Engine/Plugins/Experimental/NNE/Source/NNERuntimeRDG/Private/Hlsl/NNERuntimeRDGInstanceNormalization.cpp
+++ b/Engine/Plugins/Experimental/NNE/Source/NNERuntimeRDG/Private/Hlsl/NNERuntimeRDGInstanceNormalization.cpp
@@ -13,29 +13,29 @@ int ValidateInput(const UE::NNECore::FTensorShape& Input, const UE::NNECore::FTe
 {
 	if (Input.Rank() < 3)
 	{
-		UE_LOG(LogNNX, Warning, TEXT("InstanceNormalization input data should be at least of rank 4: %d"), Input.Rank());
+		UE_LOG(LogNNE, Warning, TEXT("InstanceNormalization input data should be at least of rank 4: %d"), Input.Rank());
 		return -1;
 	}
 
 	if (Scale.Rank() != 1)
 	{
-		UE_LOG(LogNNX, Warning, TEXT("InstanceNormalization input scale should be of rank 1: %d"), Scale.Rank());
+		UE_LOG(LogNNE, Warning, TEXT("InstanceNormalization input scale should be of rank 1: %d"), Scale.Rank());
 		return -1;
 	}
 	if (Scale.GetData()[0] != Input.GetData()[1])
 	{
-		UE_LOG(LogNNX, Warning, TEXT("InstanceNormalization input scale size should be equal to channel count: %d vs %d"), Scale.GetData()[0], Input.GetData()[1]);
+		UE_LOG(LogNNE, Warning, TEXT("InstanceNormalization input scale size should be equal to channel count: %d vs %d"), Scale.GetData()[0], Input.GetData()[1]);
 		return -1;
 	}
 
 	if (Bias.GetData()[0] != Input.GetData()[1])
 	{
-		UE_LOG(LogNNX, Warning, TEXT("InstanceNormalization intput B size should be equal to channel count: : %d vs %d"), Bias.GetData()[0], Input.GetData()[1]);
+		UE_LOG(LogNNE, Warning, TEXT("InstanceNormalization intput B size should be equal to channel count: : %d vs %d"), Bias.GetData()[0], Input.GetData()[1]);
 		return -1;
 	}
 	if (Bias.Rank() != 1)
 	{
-		UE_LOG(LogNNX, Warning, TEXT("InstanceNormalization input B should be of rank 1: %d"), Bias.Rank());
+		UE_LOG(LogNNE, Warning, TEXT("InstanceNormalization input B should be of rank 1: %d"), Bias.Rank());
 		return -1;
 	}
 
diff --git a/Engine/Plugins/Experimental/NNE/Source/NNERuntimeRDG/Private/Hlsl/NNERuntimeRDGMatMul.cpp b/Engine/Plugins/Experimental/NNE/Source/NNERuntimeRDG/Private/Hlsl/NNERuntimeRDGMatMul.cpp
--- a/Engine/Plugins/Experimental/NNE/Source/NNERuntimeRDG/Private/Hlsl/NNERuntimeRDGMatMul.cpp
+++ b/Engine/Plugins/Experimental/NNE/Source/NNERuntimeRDG/Private/Hlsl/NNERuntimeRDGMatMul.cpp
@@ -32,17 +32,17 @@ namespace UE::NNERuntimeRDG::Private::Hlsl
 
 			if (InputA.Rank() < 2)
 			{
-				UE_LOG(LogNNX, Warning, TEXT("Matmul first input should be at least of rank 2"));
+				UE_LOG(LogNNE, Warning, TEXT("Matmul first input should be at least of rank 2"));
 				return -1;
 			}
 			if (InputB.Rank() < 2)
 			{
-				UE_LOG(LogNNX, Warning, TEXT("Matmul second input should be at least of rank 2"));
+				UE_LOG(LogNNE, Warning, TEXT("Matmul second input should be at least of rank 2"));
 				return -1;
 			}
 			if (InputA.GetData()[InputA.Rank() - 1] != InputB.GetData()[InputB.Rank() - 2])
 			{
-				UE_LOG(LogNNX, Warning, TEXT("Matmul first input last dimension should be equal to second input last dimension"));
+				UE_LOG(LogNNE, Warning, TEXT("Matmul first input last dimension should be equal to second input last dimension"));
 				return -1;
 			}
 
@@ -82,12 +82,12 @@ namespace UE::NNERuntimeRDG::Private::Hlsl
 
 			if (InputA.GetShape().Rank() < 2)
 			{
-				UE_LOG(LogNNX, Warning, TEXT("Matmul first input should be at least of rank 2"));
+				UE_LOG(LogNNE, Warning, TEXT("Matmul first input should be at least of rank 2"));
 				return false;
 			}
 			if (InputB.GetShape().Rank() < 2)
 			{
-				UE_LOG(LogNNX, Warning, TEXT("Matmul second input should be at least of rank 2"));
+				UE_LOG(LogNNE, Warning, TEXT("Matmul second input should be at least of rank 2"));
 				return false;
 			}
 
diff --git a/Engine/Plugins/Experimental/NNE/Source/NNERuntimeRDG/Private/Hlsl/NNERuntimeRDGPad.cpp b/Engine/Plugins/Experimental/NNE/Source/NNERuntimeRDG/Private/Hlsl/NNERuntimeRDGPad.cpp
--- a/Engine/Plugins/Experimental/NNE/Source/NNERuntimeRDG/Private/Hlsl/NNERuntimeRDGPad.cpp
+++ b/Engine/Plugins/Experimental/NNE/Source/NNERuntimeRDG/Private/Hlsl/NNERuntimeRDGPad.cpp
@@ -43,7 +43,7 @@ namespace UE::NNERuntimeRDG::Private::Hlsl
 				int32 OutputDim = PrePad + X.GetShape().GetData()[i] + PostPad;
 				if (OutputDim < 1)
 				{
-					UE_LOG(LogNNX, Warning, TEXT("Pads cannot reduce dimension below 1, but would for tensor (name:%s) at rank %d of size %d with prepad %d and postpad %d."), *X.GetName(), i, X.GetShape().GetData()[i], PrePad, PostPad);
+					UE_LOG(LogNNE, Warning, TEXT("Pads cannot reduce dimension below 1, but would for tensor (name:%s) at rank %d of size %d with prepad %d and postpad %d."), *X.GetName(), i, X.GetShape().GetData()[i], PrePad, PostPad);
 					return -1;
 				}
 				OutputShapeData.Emplace(OutputDim);
@@ -68,7 +68,7 @@ namespace UE::NNERuntimeRDG::Private::Hlsl
 
 			if ((2*InputTensorDescs[0].GetShape().Rank()) != Pads.Num())
 			{
-				UE_LOG(LogNNX, Warning, TEXT("pads attribute lenght (%d) should be twice the rank of input X (%d)."), Pads.Num(), InputTensorDescs[0].GetShape().Rank());
+				UE_LOG(LogNNE, Warning, TEXT("pads attribute lenght (%d) should be twice the rank of input X (%d)."), Pads.Num(), InputTensorDescs[0].GetShape().Rank());
 				return false;
 			}
 
@@ -141,7 +141,7 @@ namespace UE::NNERuntimeRDG::Private::Hlsl
 		{
 			if (Pad < 0)
 			{
-				UE_LOG(LogNNX, Warning, TEXT("Pad operator does not support negative padding at the moment."));
+				UE_LOG(LogNNE, Warning, TEXT("Pad operator does not support negative padding at the moment."));
 				return false;
 			}
 		}
diff --git a/Engine/Plugins/Experimental/NNE/Source/NNERuntimeRDG/Private/Hlsl/NNERuntimeRDGUpsample.cpp b/Engine/Plugins/Experimental/NNE/Source/NNERuntimeRDG/Private/Hlsl/NNERuntimeRDGUpsample.cpp
--- a/Engine/Plugins/Experimental/NNE/Source/NNERuntimeRDG/Private/Hlsl/NNERuntimeRDGUpsample.cpp
+++ b/Engine/Plugins/Experimental/NNE/Source/NNERuntimeRDG/Private/Hlsl/NNERuntimeRDGUpsample.cpp
@@ -32,7 +32,7 @@ namespace UE::NNERuntimeRDG::Private::Hlsl
 			
 			if (!Scales.HasPreparedData())
 			{
-				UE_LOG(LogNNX, Warning, TEXT("Upsample input 'Scale' (name: %s) should be constant for shape inference to succeed, however it is not."), *Scales.GetName());
+				UE_LOG(LogNNE, Warning, TEXT("Upsample input 'Scale' (name: %s) should be constant for shape inference to succeed, however it is not."), *Scales.GetName());
 				return -1;
 			}
 
@@ -40,7 +40,7 @@ namespace UE::NNERuntimeRDG::Private::Hlsl
 
 			if (ScalesData.Num() != X.GetShape().Rank())
 			{
-				UE_LOG(LogNNX, Warning, TEXT("Upsample input 'Scale' (name: %s) have %d elements. While it should be the same as the rank of input 'X' (name : %s) witch is %d"), *Scales.GetName(), ScalesData.Num(), *X.GetName(), X.GetShape().Rank());
+				UE_LOG(LogNNE, Warning, TEXT("Upsample input 'Scale' (name: %s) have %d elements. While it should be the same as the rank of input 'X' (name : %s) witch is %d"), *Scales.GetName(), ScalesData.Num(), *X.GetName(), X.GetShape().Rank());
 				return -1;
 			}
 
@@ -127,7 +127,7 @@ namespace UE::NNERuntimeRDG::Private::Hlsl
 		FString Mode = AttributeMap.GetValueOrDefault<FString>(TEXT("mode"), TEXT("nearest"));
 		if (!Mode.Equals(TEXT("nearest")))
 		{
-			UE_LOG(LogNNX, Warning, TEXT("Upsample HLSL operator only supports nearest mode for now"));
+			UE_LOG(LogNNE, Warning, TEXT("Upsample HLSL operator only supports nearest mode for now"));
 			return false;
 		}
 
diff --git a/Engine/Plugins/Experimental/NNE/Source/NNERuntimeRDG/Private/NNERuntimeRDG.cpp b/Engine/Plugins/Experimental/NNE/Source/NNERuntimeRDG/Private/NNERuntimeRDG.cpp
--- a/Engine/Plugins/Experimental/NNE/Source/NNERuntimeRDG/Private/NNERuntimeRDG.cpp
+++ b/Engine/Plugins/Experimental/NNE/Source/NNERuntimeRDG/Private/NNERuntimeRDG.cpp
@@ -14,17 +14,6 @@
 namespace UE::NNERuntimeRDG::Private
 {
 
-FGuid FRuntimeRDG::GUID = FGuid((int32)'R', (int32)'R', (int32)'D', (int32)'G');
-int32 FRuntimeRDG::Version = 0x00000001;
-
-struct FInferenceModelRDG::FReadbackEntry
-{
-	TUniquePtr<FRHIGPUBufferReadback>	RHI;
-	void*					CpuMemory;
-	size_t					Offset;
-	size_t					Size;
-};
-
 //
 //
 //
@@ -55,12 +44,12 @@ bool FInputValidator::Validate(TConstArrayView<ENNETensorDataType> InputTypes)
 	
 	if (InputTypes.Num() < NumRequiredInput)
 	{
-		UE_LOG(LogNNX, Warning, TEXT("Required '%d' inputs but found '%d'."), NumRequiredInput, InputTypes.Num());
+		UE_LOG(LogNNE, Warning, TEXT("Required '%d' inputs but found '%d'."), NumRequiredInput, InputTypes.Num());
 		bAreInputValid = false;
 	}
 	if (InputTypes.Num() > NumRequiredInput+NumOptionalInput)
 	{
-		UE_LOG(LogNNX, Warning, TEXT("Got a total of '%d' inputs but should have '%d' maximum."), InputTypes.Num(), NumRequiredInput + NumOptionalInput);
+		UE_LOG(LogNNE, Warning, TEXT("Got a total of '%d' inputs but should have '%d' maximum."), InputTypes.Num(), NumRequiredInput + NumOptionalInput);
 		bAreInputValid = false;
 	}
 	
@@ -71,7 +60,7 @@ bool FInputValidator::Validate(TConstArrayView<ENNETensorDataType> InputTypes)
 		check(TemplateIdx < TemplateTypes.Num());
 		if (INDEX_NONE == TemplateTypes[TemplateIdx].Find(InputTypes[Idx]))
 		{
-			UE_LOG(LogNNX, Warning, TEXT("Input '%d' from template idx '%d' is of type '%d' is not supported."), Idx, TemplateIdx, (int)InputTypes[Idx]);
+			UE_LOG(LogNNE, Warning, TEXT("Input '%d' from template idx '%d' is of type '%d' is not supported."), Idx, TemplateIdx, (int)InputTypes[Idx]);
 			bAreInputValid = false;
 		}
 	}
@@ -129,13 +118,13 @@ bool FAttributeValidator::Validate(const NNECore::FAttributeMap& AttributesToVal
 		if (FoundAttribute == nullptr)
 		{
 			bAreAttributesValid = false;
-			UE_LOG(LogNNX, Warning, TEXT("Required attribute '%s' not found."),
+			UE_LOG(LogNNE, Warning, TEXT("Required attribute '%s' not found."),
 				*RequiredAttributes[Idx].Name);
 		}
 		else if (RequiredAttributes[Idx].Type != FoundAttribute->GetType())
 		{
 			bAreAttributesValid = false;
-			UE_LOG(LogNNX, Warning, TEXT("Required attribute '%s' type '%d' does not match expected type '%d'."),
+			UE_LOG(LogNNE, Warning, TEXT("Required attribute '%s' type '%d' does not match expected type '%d'."),
 				*RequiredAttributes[Idx].Name,
 				(int)FoundAttribute->GetType(),
 				(int)RequiredAttributes[Idx].Type);
@@ -150,7 +139,7 @@ bool FAttributeValidator::Validate(const NNECore::FAttributeMap& AttributesToVal
 		if ((FoundAttribute != nullptr) && (OptionalAttributes[Idx].Type != FoundAttribute->GetType()))
 		{
 			bAreAttributesValid = false;
-			UE_LOG(LogNNX, Warning, TEXT("Optional attribute '%s' type '%d' does not match expected type '%d'."),
+			UE_LOG(LogNNE, Warning, TEXT("Optional attribute '%s' type '%d' does not match expected type '%d'."),
 				*OptionalAttributes[Idx].Name,
 				(int)FoundAttribute->GetType(),
 				(int)OptionalAttributes[Idx].Type);
@@ -167,477 +156,11 @@ bool FAttributeValidator::Validate(const NNECore::FAttributeMap& AttributesToVal
 		if (OptionalAttribute == nullptr && RequiredAttribute == nullptr)
 		{
 			bAreAttributesValid = false;
-			UE_LOG(LogNNX, Warning, TEXT("Found unsupported attribute '%s'."), *Name);
+			UE_LOG(LogNNE, Warning, TEXT("Found unsupported attribute '%s'."), *Name);
 		}
 	}
 
 	return bAreAttributesValid;
 }
 
-bool FRuntimeRDG::CanCreateModelData(FString FileType, TConstArrayView<uint8> FileData) const
-{
-	return FileType.Compare("onnx", ESearchCase::IgnoreCase) == 0;
-}
-
-bool FRuntimeRDG::CanCreateModel(TConstArrayView<uint8> ModelData) const
-{
-	int32 GuidSize = sizeof(FRuntimeRDG::GUID);
-	int32 VersionSize = sizeof(FRuntimeRDG::Version);
-	if (ModelData.Num() <= GuidSize + VersionSize)
-	{
-		return false;
-	}
-	bool bResult = FGenericPlatformMemory::Memcmp(&(ModelData[0]), &(FRuntimeRDG::GUID), GuidSize) == 0;
-	bResult &= FGenericPlatformMemory::Memcmp(&(ModelData[GuidSize]), &(FRuntimeRDG::Version), VersionSize) == 0;
-	return bResult;
-};
-	
-//
-//
-//
-FInferenceModelRDG::FInferenceModelRDG()
-	: FMLInferenceModel(EMLInferenceModelType::RDG)
-	, bUseManualTransitions(false)
-{
-	
-}
-
-FInferenceModelRDG::~FInferenceModelRDG()
-{
-	
-}
-
-//
-//
-//
-bool FInferenceModelRDG::LoadModel(TConstArrayView<uint8> ModelData, FMLRuntimeFormat& Format)
-{
-	int32 GuidSize = sizeof(FRuntimeRDG::GUID);
-	int32 VersionSize = sizeof(FRuntimeRDG::Version);
-	TConstArrayView<uint8> ModelBuffer = {&(ModelData.GetData()[GuidSize + VersionSize]), ModelData.Num() - GuidSize - VersionSize};
-
-	FMemoryReaderView Reader(ModelBuffer);
-
-	FMLRuntimeFormat::StaticStruct()->SerializeBin(Reader, &Format);
-
-	// Data for base class
-	InputSymbolicTensors.Empty();
-	OutputSymbolicTensors.Empty();
-	
-	// Data for RDG
-	AllSymbolicTensorDescs.Empty();
-	IntermediateTensorIndices.Empty();
-	WeightTensorIndices.Empty();
-	InputTensorIndices.Empty();
-	OutputTensorIndices.Empty();
-	OperatorInputTensorIndices.Empty();
-	OperatorOutputTensorIndices.Empty();
-
-	// Add tensors
-	for (int32 Idx = 0; Idx < Format.Tensors.Num(); ++Idx)
-	{
-		const FMLFormatTensorDesc& FormatTensorDesc = Format.Tensors[Idx];
-
-		const NNECore::FSymbolicTensorShape SymbolicShape = NNECore::FSymbolicTensorShape::Make(FormatTensorDesc.Shape);
-		const NNECore::FTensorDesc SymbolicTensor = NNECore::FTensorDesc::Make(FormatTensorDesc.Name, SymbolicShape, FormatTensorDesc.DataType);
-		
-		AllSymbolicTensorDescs.Emplace(SymbolicTensor);
-		
-		if (FormatTensorDesc.Type == EMLFormatTensorType::Input)
-		{
-			InputTensorIndices.Emplace(Idx);
-			InputSymbolicTensors.Emplace(SymbolicTensor);
-		}
-		else if (FormatTensorDesc.Type == EMLFormatTensorType::Output)
-		{
-			OutputTensorIndices.Emplace(Idx);
-			OutputSymbolicTensors.Emplace(SymbolicTensor);
-		}
-		else if (FormatTensorDesc.Type == EMLFormatTensorType::Intermediate)
-		{
-			IntermediateTensorIndices.Emplace(Idx);
-		}
-		else if (FormatTensorDesc.Type == EMLFormatTensorType::Initializer)
-		{
-			WeightTensorIndices.Emplace(Idx);
-			if (!SymbolicTensor.GetShape().IsConcrete())
-			{
-				UE_LOG(LogNNX, Error, TEXT("Weight tensor %s should have a concrete shape"), *SymbolicTensor.GetName());
-				return false;
-			}
-			
-			const NNECore::FTensorShape TensorShape = NNECore::FTensorShape::MakeFromSymbolic(SymbolicTensor.GetShape());
-			FTensorRDG& WeightRDG = WeightTensorRDGs.Emplace_GetRef(FTensorRDG::Make(SymbolicTensor, TensorShape, nullptr));
-			
-			if (WeightRDG.GetDataSize() != FormatTensorDesc.DataSize)
-			{
-				UE_LOG(LogNNX, Error, TEXT("Weight %s has incorrect size. Expected %d bytes, got %d"), *SymbolicTensor.GetName(), FormatTensorDesc.DataSize, WeightRDG.GetDataSize());
-				return false;
-			}
-
-			const uint8* DataPtr = Format.TensorData.GetData() + FormatTensorDesc.DataOffset;
-			TConstArrayView<uint8> DataView = MakeArrayView(DataPtr, FormatTensorDesc.DataSize);
-
-			WeightRDG.SetPreparedData(DataView);
-		}
-		checkf(FormatTensorDesc.Type != EMLFormatTensorType::None, TEXT("Unsupported tensor type None"));
-	}
-
-	// Loop over all operators in the model and store tensor indices for input/output
-	for (int32 Idx = 0; Idx < Format.Operators.Num(); ++Idx)
-	{
-		OperatorInputTensorIndices.Emplace(Format.Operators[Idx].InTensors);
-		OperatorOutputTensorIndices.Emplace(Format.Operators[Idx].OutTensors);
-	}
-
-	return true;
-}
-
-/**
- * Run the inference model (synchronous version)
- */
-int FInferenceModelRDG::RunSync(TConstArrayView<FMLTensorBinding> InInputBindings, TConstArrayView<FMLTensorBinding> InOutputBindings)
-{
-	// Verify the model inputs were prepared
-	if (InputTensorShapes.Num() == 0)
-	{
-		UE_LOG(LogNNX, Error, TEXT("Run(): Input shapes are not set, please call SetInputTensorShapes."));
-		return -1;
-	}
-
-	// TODO ok if pending?
-	Readbacks.Empty();
-	
-	int Res = 0;
-	FEvent* Signal = FGenericPlatformProcess::GetSynchEventFromPool(false);
-
-	ENQUEUE_RENDER_COMMAND(FMLInferenceModel_Run)
-	(
-		[&Signal, &Res, this, InInputBindings, InOutputBindings](FRHICommandListImmediate& RHICmdList)
-		{
-			TOptional<ERHIPipeline>		Pipeline = RHICmdList.GetPipeline();
-
-			if (Pipeline == ERHIPipeline::None)
-			{
-				RHICmdList.SwitchPipeline(ERHIPipeline::Graphics);
-			}
-
-			FRDGBuilder	RDGBuilder(RHICmdList);
-
-			Res = EnqueueRDG(RDGBuilder, InInputBindings, InOutputBindings);
-			if (Res == 0)
-			{
-				RDGBuilder.Execute();
-
-				// FIXME: Using BlockUntilGPUIdle() prevents hang on Linux
-				// FIXME: Adapt to redesigned readback API (UE 5.2)
-				RHICmdList.BlockUntilGPUIdle();
-
-				// Process readback
-				for (const FReadbackEntry& Readback : Readbacks) {
-					const void* BuffData = Readback.RHI->Lock(Readback.Size);
-					check(BuffData);
-					FMemory::Memcpy(Readback.CpuMemory, BuffData, Readback.Size);
-					Readback.RHI->Unlock();
-				}
-			}
-			
-			Signal->Trigger();
-		}
-	);
-
-	// We need to wait for render thread to finish
-	Signal->Wait();
-
-	FGenericPlatformProcess::ReturnSynchEventToPool(Signal);
-
-	return Res;
-}
-
-int FInferenceModelRDG::SetInputTensorShapes(TConstArrayView<FTensorShape> InInputShapes)
-{
-	OutputTensorShapes.Empty();
-	
-	//Verify input shape are valid for the model and set InputTensorShapes
-	if (FMLInferenceModel::SetInputTensorShapes(InInputShapes) != 0)
-	{
-		return -1;
-	}
-
-	//Allocate and prime all AllTensorRDGs with concrete shapes defaulting variables dimension to 1 if needed
-	AllTensorRDGs.Init(nullptr, AllSymbolicTensorDescs.Num());
-
-	InputTensorRDGs.Empty();
-	for (int32 i = 0; i < InputTensorIndices.Num(); ++i)
-	{
-		const int32 Idx = InputTensorIndices[i];
-		const FTensorDesc& TensorDesc = InputSymbolicTensors[i];
-		const FTensorShape& TensorShape = InputTensorShapes[i];
-		
-		InputTensorRDGs.Emplace(FTensorRDG::Make(TensorDesc, TensorShape, nullptr));
-		AllTensorRDGs[Idx] = &InputTensorRDGs[i];
-	}
-
-	for (int32 i = 0; i < WeightTensorIndices.Num(); ++i)
-	{
-		const int32 Idx = WeightTensorIndices[i];
-		
-		AllTensorRDGs[Idx] = &WeightTensorRDGs[i];
-	}
-
-	IntermediateTensorRDGs.Empty();
-	for (int32 i = 0; i < IntermediateTensorIndices.Num(); ++i)
-	{
-		const int32 Idx = IntermediateTensorIndices[i];
-		const FTensorDesc& TensorDesc = AllSymbolicTensorDescs[Idx];
-		const FTensorShape TensorShape = FTensorShape::MakeFromSymbolic(TensorDesc.GetShape());
-		
-		IntermediateTensorRDGs.Emplace(FTensorRDG::Make(TensorDesc, TensorShape, nullptr));
-		AllTensorRDGs[Idx] = &IntermediateTensorRDGs[i];
-	}
-
-	OutputTensorRDGs.Empty();
-	for (int32 i = 0; i < OutputTensorIndices.Num(); ++i)
-	{
-		const int32 Idx = OutputTensorIndices[i];
-		const FTensorDesc& TensorDesc = OutputSymbolicTensors[i];
-		const FTensorShape TensorShape = FTensorShape::MakeFromSymbolic(TensorDesc.GetShape());
-		
-		OutputTensorRDGs.Emplace(FTensorRDG::Make(TensorDesc, TensorShape, nullptr));
-		AllTensorRDGs[Idx] = &OutputTensorRDGs[i];
-	}
-
-	checkCode(
-		for (int i = 0; i < AllTensorRDGs.Num(); ++i)
-		{
-			checkf(AllTensorRDGs[i] != nullptr, TEXT("Tensor at index %d, was not allocated for model preparation."), i);
-		};
-	);
-
-	//Allow the specific engine to run shape inference if supported
-	if (PrepareTensorShapesAndData() != 0)
-	{
-		return -1;
-	}
-
-	checkCode(
-		for (int i = 0; i < AllTensorRDGs.Num(); ++i)
-		{
-			checkf(AllTensorRDGs[i] != nullptr, TEXT("Tensor at index %d, was not allocated after model preparation."), i);
-			checkf(AllTensorRDGs[i]->GetShape().IsCompatibleWith(AllSymbolicTensorDescs[i].GetShape()), TEXT("Tensor at index %d have a shape incompatible with model definition."), i);
-		};
-	);
-
-	//Set OutputTensorShapes for the model from preparation result
-	for (int32 OutputIndices : OutputTensorIndices)
-	{
-		OutputTensorShapes.Emplace(AllTensorRDGs[OutputIndices]->GetShape());
-	}
-
-	check(InputTensorIndices.Num() + OutputTensorIndices.Num() + WeightTensorIndices.Num() + IntermediateTensorIndices.Num() == AllTensorRDGs.Num());
-	check(InputTensorShapes.Num() == InputSymbolicTensors.Num());
-	check(OutputTensorShapes.Num() == OutputSymbolicTensors.Num());
-	check(WeightTensorIndices.Num() == WeightTensorRDGs.Num());
-	check(AllTensorRDGs.Num() == AllSymbolicTensorDescs.Num());
-	
-	return 0;
-}
-
-FRDGBufferDesc CreateRDGBufferDescForTensorRDG(const FTensorRDG& Tensor)
-{
-	// FIXME: CreateStructuredDesc() creates a crash on VulkanRHI
-	//FRDGBufferDesc Desc = FRDGBufferDesc::CreateStructuredDesc(Tensor.GetElemByteSize(), Tensor.GetVolume());
-	FRDGBufferDesc Desc = FRDGBufferDesc::CreateBufferDesc(Tensor.GetElemByteSize(), Tensor.GetVolume());
-
-	return Desc;
-}
-
-/**
- * Enqueue operators to RDG, the caller will run the GraphBuilder.Execute()
- */
-int FInferenceModelRDG::EnqueueRDG(FRDGBuilder& RDGBuilder, TConstArrayView<FMLTensorBinding> InInputBindings, TConstArrayView<FMLTensorBinding> InOutputBindings)
-{
-	check(IsInRenderingThread());
-
-	int Res;
-
-	// Verify the model inputs were prepared
-	if (InputTensorShapes.Num() == 0)
-	{
-		UE_LOG(LogNNX, Error, TEXT("EnqueueRDG(): Input shapes are not set, please call SetInputTensorShapes."));
-		return -1;
-	}
-
-	// Process input binding, and if required, allocate RDG buffers
-	FIntArray		InputUploadIndices;
-	Res = SetTensors(RDGBuilder, InputTensorRDGs, InputUploadIndices, InInputBindings);
-	if (Res != 0)
-	{
-		UE_LOG(LogNNX, Warning, TEXT("Invalid input tensor binding type for tensor index:%d"), Res);
-		return -1;
-	}
-
-	// Process output tensors, and if required, allocate RDG buffers
-	FIntArray		OutputReadbackIndices;
-
-	Res = SetTensors(RDGBuilder, OutputTensorRDGs, OutputReadbackIndices, InOutputBindings);
-	if (Res != 0)
-	{
-		UE_LOG(LogNNX, Warning, TEXT("Invalid output tensor binding type for tensor index:%d"), Res);
-		return -1;
-	}
-
-	// If required, upload input tensors to GPU
-	AddTensorUploads_RenderThread(RDGBuilder, InputUploadIndices, InputTensorRDGs, InInputBindings);
-
-	//Create buffer for intermediate tensors
-	for (FTensorRDG& TensorRDG : IntermediateTensorRDGs)
-	{
-		const FRDGBufferDesc BufferDesc = CreateRDGBufferDescForTensorRDG(TensorRDG);
-		const FRDGBufferRef TensorBuffer = RDGBuilder.CreateBuffer(BufferDesc, *TensorRDG.GetName(), ERDGBufferFlags::None);
-		check(TensorRDG.GetBuffer() == nullptr);
-		TensorRDG.SetBuffer(TensorBuffer);
-	}
-
-	// TODO: FIXME: DirectML uses RHI buffers instead of RDG bufers
-	//For now weights tensors are not uploaded to GPU thus GetBuffer will return nullptr for them.
-	//checkCode(for (const FTensorRDG* TensorRDG : AllTensorRDGs) { if (TensorRDG != nullptr) { check(TensorRDG->GetBuffer() != nullptr); } });
-
-	//Insert weights tensors
-	for (int32 i = 0; i < WeightTensorIndices.Num(); ++i)
-	{
-		AllTensorRDGs[WeightTensorIndices[i]] = &WeightTensorRDGs[i];
-	}
-	
-	// We can now dispatch operators
-	AddDispatchOps_RenderThread(RDGBuilder);
-
-	// If required, readback the output tensors to CPU
-	AddTensorReadbacks_RenderThread(RDGBuilder, OutputReadbackIndices, OutputTensorRDGs, InOutputBindings);
-
-	return 0;
-}
-
-/** 
- * Process binding and check if we need to create RDG Buffer for CPU binding 
- * Returns 0 on success, or index of a tensor if the tensor type is not supported.
- */
-int FInferenceModelRDG::SetTensors(FRDGBuilder& GraphBuilder, FTensorRDGArray& InTensorRDGs, FIntArray& OutIndices, 
-	TConstArrayView<FMLTensorBinding> InBindings)
-{
-	check(InBindings.Num() == InTensorRDGs.Num());
-	
-	for (int32 Idx = 0; Idx < InBindings.Num(); ++Idx)
-	{
-		FTensorRDG& TensorRDG = InTensorRDGs[Idx];
-		const FMLTensorBinding& Binding = InBindings[Idx];
-
-		if (Binding.BindingType == EMLTensorBindingDataType::CPUMemory)
-		{
-			FRDGBufferDesc Desc = CreateRDGBufferDescForTensorRDG(TensorRDG);
-			
-			// FIXME: We should use BUF_SourceCopy for only output buffers (GPU readback)
-			Desc.Usage = EBufferUsageFlags(Desc.Usage | BUF_SourceCopy);
-			
-			FRDGBufferRef TensorBuffer = GraphBuilder.CreateBuffer(Desc, *TensorRDG.GetName(), ERDGBufferFlags::None);
-			
-			TensorRDG.SetBuffer(TensorBuffer);
-			OutIndices.Add(Idx);
-		}
-		else if (Binding.BindingType == EMLTensorBindingDataType::RDGBuffer)
-		{
-			TensorRDG.SetBuffer(Binding.Buffer);
-		}
-		else
-		{
-			// Unsupported tensor binding type
-			return Idx;
-		}
-	}
-
-	return 0;
-}
-
-//
-//
-//
-void FInferenceModelRDG::AddTensorUploads_RenderThread(FRDGBuilder& GraphBuilder, TConstArrayView<int32> InUploadIndices, 
-	TConstArrayView<FTensorRDG> InTensorRDGs, TConstArrayView<FMLTensorBinding> InBindings)
-{
-	check(InTensorRDGs.Num() == InBindings.Num());
-	
-	for (int32 Idx = 0; Idx < InUploadIndices.Num(); ++Idx)
-	{
-		const int32			TensorIdx = InUploadIndices[Idx];
-		const FTensorRDG&	TensorRDG = InTensorRDGs[TensorIdx];
-		const FMLTensorBinding& Binding = InBindings[TensorIdx];
-		check(Binding.BindingType == EMLTensorBindingDataType::CPUMemory);
-		GraphBuilder.QueueBufferUpload(TensorRDG.GetBuffer(), Binding.CpuMemory, TensorRDG.GetDataSize(), ERDGInitialDataFlags::NoCopy);
-	}
-}
-
-//
-//
-//
-void FInferenceModelRDG::AddTensorReadbacks_RenderThread(FRDGBuilder& GraphBuilder, TConstArrayView<int32> InReadbackIndices,
-	TConstArrayView<FTensorRDG> InTensorRDGs, TConstArrayView<FMLTensorBinding> InBindings)
-{
-	check(InTensorRDGs.Num() == InBindings.Num());
-
-	check(Readbacks.IsEmpty());
-	
-	for (int32 Idx = 0; Idx < InReadbackIndices.Num(); ++Idx)
-	{
-		const int32			TensorIdx = InReadbackIndices[Idx];
-		const FTensorRDG&	TensorRDG = InTensorRDGs[TensorIdx];
-		const FMLTensorBinding& Binding = InBindings[TensorIdx];
-		check(Binding.BindingType == EMLTensorBindingDataType::CPUMemory);
-		
-		FMLTensorReadbackParameters* TensorReadbackParams = GraphBuilder.AllocParameters<FMLTensorReadbackParameters>();
-
-		TensorReadbackParams->Buffer = TensorRDG.GetBuffer();
-
-		GraphBuilder.AddPass(
-			RDG_EVENT_NAME("FMLInferenceModelAddTensorReadback:%s", *TensorRDG.GetName()),
-			TensorReadbackParams,
-			ERDGPassFlags::Readback | ERDGPassFlags::NeverCull,
-			[this, Binding, TensorRDG, TensorReadbackParams](FRHICommandListImmediate& RHICmdList)
-			{
-				FRHIBuffer* OutputBuffer = TensorReadbackParams->Buffer->GetRHI();
-
-				// TODO: FIXME: We need to transition the resources for DirectML
-				if (bUseManualTransitions)
-				{
-					FRHITransitionInfo Transitions[] =
-					{
-						FRHITransitionInfo(OutputBuffer, ERHIAccess::UAVCompute, ERHIAccess::CopySrc)
-					};
-
-					RHICmdList.Transition(MakeArrayView(Transitions, UE_ARRAY_COUNT(Transitions)));
-					RHICmdList.SubmitCommandsHint();
-				}
-
-				FReadbackEntry& Readback = Readbacks.Add_GetRef({});
-				Readback.RHI = MakeUnique<FRHIGPUBufferReadback>(FName(TEXT("FMLTensorReadback_") + TensorRDG.GetName()));
-				Readback.RHI->EnqueueCopy(RHICmdList, OutputBuffer, TensorRDG.GetDataSize());
-				Readback.CpuMemory = Binding.CpuMemory;
-				Readback.Offset = 0;
-				Readback.Size = TensorRDG.GetDataSize();
-			}
-		);
-	}
-}
-
-TArray<uint8> ConvertToModelData(TArrayView<uint8> ModelBuffer)
-{
-	TArray<uint8> Result;
-
-	FMemoryWriter Writer(Result);
-	Writer << FRuntimeRDG::GUID;
-	Writer << FRuntimeRDG::Version;
-	Writer.Serialize(ModelBuffer.GetData(), ModelBuffer.Num());
-
-	return Result;
-}
-
 } // UE::NNERuntimeRDG::Private
diff --git a/Engine/Plugins/Experimental/NNE/Source/NNERuntimeRDG/Private/NNERuntimeRDG.h b/Engine/Plugins/Experimental/NNE/Source/NNERuntimeRDG/Private/NNERuntimeRDG.h
--- a/Engine/Plugins/Experimental/NNE/Source/NNERuntimeRDG/Private/NNERuntimeRDG.h
+++ b/Engine/Plugins/Experimental/NNE/Source/NNERuntimeRDG/Private/NNERuntimeRDG.h
@@ -6,31 +6,20 @@
 #include "NNECoreAttributeValue.h"
 #include "NNECoreTensor.h"
 #include "NNECoreTypes.h"
-#include "NNXCore.h"
+#include "NNECore.h"
 #include "NNXModelOptimizerInterface.h"
-#include "NNXRuntime.h"
 #include "NNXRuntimeFormat.h"
 #include "RenderGraphResources.h"
 #include "Serialization/MemoryReader.h"
 #include "ShaderParameterUtils.h"
 
-#include "Containers/Map.h"
 
-BEGIN_SHADER_PARAMETER_STRUCT(FMLTensorUploadParameters, )
-	RDG_BUFFER_ACCESS(Buffer, ERHIAccess::CopyDest)
-END_SHADER_PARAMETER_STRUCT()
 
-BEGIN_SHADER_PARAMETER_STRUCT(FMLTensorReadbackParameters, )
+BEGIN_SHADER_PARAMETER_STRUCT(FNNETensorReadbackParameters, )
 	RDG_BUFFER_ACCESS(Buffer, ERHIAccess::CopySrc)
 END_SHADER_PARAMETER_STRUCT()
 
-BEGIN_SHADER_PARAMETER_STRUCT(FMLElementWiseOpParameters, )
-	RDG_BUFFER_ACCESS(InputBuffer, ERHIAccess::UAVCompute)			//!< NOTE: DirectML requires state to be in UAV, even though we're just reading from the InputBuffer
-	RDG_BUFFER_ACCESS(OutputBuffer, ERHIAccess::UAVCompute)
-END_SHADER_PARAMETER_STRUCT()
-
 class FRDGBuilder;
-struct FMLRuntimeFormat;
 struct FNNIModelRaw;
 
 namespace UE::NNECore { class FTensorDesc; }
@@ -40,8 +29,6 @@ using namespace NNX;
 namespace UE::NNERuntimeRDG::Private
 {
 
-struct FTensorBinding;
-
 /**
  * Interface for all operators to prepare the model tensors at scheduling time
  */
@@ -89,72 +76,7 @@ using FTensorRDGArray = TArray<FTensorRDG, TInlineAllocator<16>>;
 using FTensorRDGRefArray = TArray<FTensorRDGRef, TInlineAllocator<64>>;
 using FIntArray = TArray<int32, TInlineAllocator<16>>;
 
-class FRuntimeRDG : public IRuntime
-{
-public:
-	static FGuid GUID;
-	static int32 Version;
-
-	virtual bool CanCreateModelData(FString FileType, TConstArrayView<uint8> FileData) const;
-	virtual TArray<uint8> CreateModelData(FString FileType, TConstArrayView<uint8> FileData) = 0;
-	virtual bool CanCreateModel(TConstArrayView<uint8> ModelData) const;
-};
-
-/** 
- * RDG inference model base class
- */
-class FInferenceModelRDG : public FMLInferenceModel
-{
-
-	struct FReadbackEntry;
-
-public:
-
-	~FInferenceModelRDG();
-
-	virtual int SetInputTensorShapes(TConstArrayView<FTensorShape> InputShapes) override;
-	virtual int RunSync(TConstArrayView<FMLTensorBinding> InputBindings, TConstArrayView<FMLTensorBinding> OutputBindings) override;
-	virtual int EnqueueRDG(FRDGBuilder& RDGBuilder, TConstArrayView<FMLTensorBinding> InputBindings, TConstArrayView<FMLTensorBinding> OutputBindings) override;
-
-protected:
-
-	FInferenceModelRDG();
-
-	bool LoadModel(TConstArrayView<uint8> ModelData, FMLRuntimeFormat& Format);
-
-	int SetTensors(FRDGBuilder& GraphBuilder, FTensorRDGArray& InTensorRDGs, FIntArray& OutIndices, TConstArrayView<FMLTensorBinding> InBindings);
-	
-	virtual int PrepareTensorShapesAndData() = 0;
-	virtual void AddDispatchOps_RenderThread(FRDGBuilder& GraphBuilder) = 0;
-
-	virtual void AddTensorUploads_RenderThread(FRDGBuilder& GraphBuilder, TConstArrayView<int32> InUploadIndices, TConstArrayView<FTensorRDG> InRDGBindings, TConstArrayView<FMLTensorBinding> InBindings);
-	virtual void AddTensorReadbacks_RenderThread(FRDGBuilder& GraphBuilder, TConstArrayView<int32> InReadbackIndices, TConstArrayView<FTensorRDG> InRDGBindings, TConstArrayView<FMLTensorBinding> InBindings);
-
-	//Tensor descriptor
-	TArray<FTensorDesc>			AllSymbolicTensorDescs;
-
-	//Tensor indices for models
-	TArray<int32>				IntermediateTensorIndices;
-	TArray<int32>				WeightTensorIndices;
-	TArray<int32>				InputTensorIndices;
-	TArray<int32>				OutputTensorIndices;
-	
-	//Tensor indices by operator
-	TArray<TArray<uint32>>		OperatorInputTensorIndices;
-	TArray<TArray<uint32>>		OperatorOutputTensorIndices;
-
-	//RDG Tensors
-	FTensorRDGRefArray			AllTensorRDGs;
-	FTensorRDGArray				InputTensorRDGs;
-	FTensorRDGArray				OutputTensorRDGs;
-	FTensorRDGArray				IntermediateTensorRDGs;
-	FTensorRDGArray				WeightTensorRDGs;
-	
-	TArray<FReadbackEntry>		Readbacks;
-	bool						bUseManualTransitions;
-};
-
-//TODO jira 167584 remove default validation and declare contract in all DML operator (see HLSL Gemm for current example)
+//Note: jira 167584 remove default validation and declare contract in all DML operator (see HLSL Gemm for current example)
 bool AlwaysValidValidationFunction(const NNECore::FAttributeMap& AttributeMap, TConstArrayView<ENNETensorDataType> InputTypes, TConstArrayView<NNECore::FSymbolicTensorShape> InputShapes);
 
 class FInputValidator
@@ -225,7 +147,7 @@ public:
 
 		if (!Fn)
 		{
-			UE_LOG(LogNNX, Warning, TEXT("RDG MLOperator:%s is not registered"), *Name);
+			UE_LOG(LogNNE, Warning, TEXT("RDG MLOperator:%s is not registered"), *Name);
 			return nullptr;
 		}
 
@@ -238,7 +160,7 @@ public:
 
 		if (!Fn)
 		{
-			UE_LOG(LogNNX, Warning, TEXT("RDG MLOperator:%s is not registered"), *Name);
+			UE_LOG(LogNNE, Warning, TEXT("RDG MLOperator:%s is not registered"), *Name);
 			return nullptr;
 		}
 
@@ -249,7 +171,7 @@ public:
 	{
 		if (Operators.Find(Name) != nullptr)
 		{
-			UE_LOG(LogNNX, Warning, TEXT("RDG MLOperator is already registered:%s"), *Name);
+			UE_LOG(LogNNE, Warning, TEXT("RDG MLOperator is already registered:%s"), *Name);
 			return false;
 		}
 
@@ -283,7 +205,7 @@ public:
 		ENNXInferenceFormat FormatType = InputModel.Format;
 		if (FormatType != ENNXInferenceFormat::NNXRT)
 		{
-			UE_LOG(LogNNX, Warning, TEXT("Unsupported format type for validator %s"), *GetName());
+			UE_LOG(LogNNE, Warning, TEXT("Unsupported format type for validator %s"), *GetName());
 			return false;
 		}
 
@@ -315,13 +237,13 @@ public:
 
 			if (!ValidationFn)
 			{
-				UE_LOG(LogNNX, Warning, TEXT("Hlsl MLOperatorRegistry failed to find validation for operator:%s"), *OpType);
+				UE_LOG(LogNNE, Warning, TEXT("Hlsl MLOperatorRegistry failed to find validation for operator:%s"), *OpType);
 				return false;
 			}
 			
 			if (!ValidationFn(AttributeMap, InputTensorTypes, InputTensorShapes))
 			{
-				UE_LOG(LogNNX, Warning, TEXT("Hlsl MLOperatorRegistry failed to validate operator:%s"), *OpType);
+				UE_LOG(LogNNE, Warning, TEXT("Hlsl MLOperatorRegistry failed to validate operator:%s"), *OpType);
 				return false;
 			}
 		}
@@ -330,69 +252,4 @@ public:
 	}
 };
 
-TArray<uint8> ConvertToModelData(TArrayView<uint8> ModelBuffer);
-
-} // UE::NNERuntimeRDG::Private
-
-namespace UE::NNERuntimeRDG::Private::Dml
-{
-
-// NOTE: For now we only have DML on Windows, we should add support for XSX
-#ifdef NNE_USE_DIRECTML
-
-	extern IRuntime* FRuntimeDmlStartup();
-	extern void FRuntimeDmlShutdown();
-
-#endif
-	
-}
-
-namespace UE::NNERuntimeRDG::Private::Hlsl
-{
-
-	extern IRuntime* FRuntimeHlslStartup();
-	extern void FRuntimeHlslShutdown();
-}
-
-///// OperatorId
-//class FOperatorId
-//{
-//public:
-//
-//	FOperatorId(uint32 InType, uint32 InIndex)
-//		: Type(InType)
-//		, Index(InIndex)
-//	{
-//	}
-//
-//	uint32 GetIndex() const
-//	{
-//		return Index;
-//	}
-//
-//	uint32 GetType() const
-//	{
-//		return Type;
-//	}
-//
-//private:
-//
-//	uint32	Type : 8;
-//	uint32	Index : 16;
-//};
-
-///// Base class
-//class FRDGOperator
-//{};
-
-///// Base class for RDG runtime
-//class FRDGRuntime : public IRuntime
-//{
-//public:
-//
-//	virtual FRDGOperator OpCreate(const FString& Name, const TArrayView<uint32>& TensorSizes) = 0;
-//
-//};
-
-
- 
\ No newline at end of file
+} // UE::NNERuntimeRDG::Private
\ No newline at end of file
diff --git a/Engine/Plugins/Experimental/NNE/Source/NNERuntimeRDG/Private/NNERuntimeRDGDml.cpp b/Engine/Plugins/Experimental/NNE/Source/NNERuntimeRDG/Private/NNERuntimeRDGDml.cpp
--- a/Engine/Plugins/Experimental/NNE/Source/NNERuntimeRDG/Private/NNERuntimeRDGDml.cpp
+++ b/Engine/Plugins/Experimental/NNE/Source/NNERuntimeRDG/Private/NNERuntimeRDGDml.cpp
@@ -1,18 +1,16 @@
 // Copyright Epic Games, Inc. All Rights Reserved.
 
+#include "NNERuntimeRDGDml.h"
 #include "NNERuntimeRDG.h"
-#include "NNXInferenceModel.h"
+#include "NNERuntimeRDGModel.h"
 #include "NNXRuntimeFormat.h"
 #include "NNEUtilsModelOptimizer.h"
 #include "NNECoreOperator.h"
-
 #include "NNECoreAttributeMap.h"
-
+#include "NNECoreModelData.h"
 #include "HAL/FileManager.h"
-
 #include "RenderGraphBuilder.h"
 #include "RenderGraphUtils.h"
-
 #include "Algo/Find.h"
 
 #ifdef NNE_USE_DIRECTML
@@ -35,8 +33,6 @@ THIRD_PARTY_INCLUDES_END
 #include "ID3D12DynamicRHI.h"
 
 #define NNE_USE_D3D12_RESOURCES
-#define NNX_RUNTIME_DML_NAME TEXT("NNXRuntimeDml")
-
 
 namespace UE::NNERuntimeRDG::Private::Dml
 {
@@ -386,7 +382,7 @@ protected:
 			DmlUtil::SetTensorSizesAndStridesForBroadcast(DmlTensorDesc, TensorDesc, BroadcastDesc);
 		}
 
-		//UE_LOG(LogNNX, Warning, TEXT("DmlTensorDesc:%d,%d,%d -> %d,%d,%d"),
+		//UE_LOG(LogNNE, Warning, TEXT("DmlTensorDesc:%d,%d,%d -> %d,%d,%d"),
 		//	TensorDesc.Sizes[0],
 		//	TensorDesc.Dimension > 1 ? TensorDesc.Sizes[1] : 0,
 		//	TensorDesc.Dimension > 2 ? TensorDesc.Sizes[2] : 0,
@@ -395,7 +391,7 @@ protected:
 		//	DmlTensorDesc.Dimension > 2 ? DmlTensorDesc.Sizes[2] : 0
 		//);
 
-		//UE_LOG(LogNNX, Warning, TEXT("DmlTensorStrides:%d,%d,%d"),
+		//UE_LOG(LogNNE, Warning, TEXT("DmlTensorStrides:%d,%d,%d"),
 		//	DmlTensorDesc.Strides[0], DmlTensorDesc.Strides[1], DmlTensorDesc.Strides[2]
 		//);
 
@@ -429,7 +425,7 @@ protected:
 		Res = Device->CreateOperator(&DmlOpDesc, DML_PPV_ARGS(&Op));
 		if (!Op)
 		{
-			UE_LOG(LogNNX, Warning, TEXT("Error:Failed to create DML operator, hres:%d"), Res);
+			UE_LOG(LogNNE, Warning, TEXT("Error:Failed to create DML operator, hres:%d"), Res);
 			return false;
 		}
 
@@ -497,7 +493,7 @@ public:
 
 		if (!InitDmlTensorDesc(DmlTensorDesc, InputTensorDesc))
 		{
-			UE_LOG(LogNNX, Warning, TEXT("Failed to initialize tensor(s) for DML inference"));
+			UE_LOG(LogNNE, Warning, TEXT("Failed to initialize tensor(s) for DML inference"));
 			return false;
 		}
 
@@ -627,19 +623,19 @@ public:
 
 		if (!InitDmlTensorDesc(DmlInputATensorDesc, InputATensorDesc, OutputTensorDesc))
 		{
-			UE_LOG(LogNNX, Warning, TEXT("Failed to initialize tensor(s) for DML inference"));
+			UE_LOG(LogNNE, Warning, TEXT("Failed to initialize tensor(s) for DML inference"));
 			return false;
 		}
 
 		if (!InitDmlTensorDesc(DmlInputBTensorDesc, InputBTensorDesc, OutputTensorDesc))
 		{
-			UE_LOG(LogNNX, Warning, TEXT("Failed to initialize tensor(s) for DML inference"));
+			UE_LOG(LogNNE, Warning, TEXT("Failed to initialize tensor(s) for DML inference"));
 			return false;
 		}
 
 		if (!InitDmlTensorDesc(DmlOutputTensorDesc, OutputTensorDesc))
 		{
-			UE_LOG(LogNNX, Warning, TEXT("Failed to initialize tensor(s) for DML inference"));
+			UE_LOG(LogNNE, Warning, TEXT("Failed to initialize tensor(s) for DML inference"));
 			return false;
 		}
 
@@ -722,13 +718,13 @@ public:
 
 		if (!InitDmlTensorDesc(DmlInputATensorDesc, InputATensorDesc))
 		{
-			UE_LOG(LogNNX, Warning, TEXT("Failed to initialize tensor(s) for DML inference"));
+			UE_LOG(LogNNE, Warning, TEXT("Failed to initialize tensor(s) for DML inference"));
 			return false;
 		}
 
 		if (!InitDmlTensorDesc(DmlInputBTensorDesc, InputBTensorDesc))
 		{
-			UE_LOG(LogNNX, Warning, TEXT("Failed to initialize tensor(s) for DML inference"));
+			UE_LOG(LogNNE, Warning, TEXT("Failed to initialize tensor(s) for DML inference"));
 			return false;
 		}
 
@@ -738,14 +734,14 @@ public:
 
 			if (!InitDmlTensorDesc(DmlInputCTensorDesc, InputCTensorDesc, OutputTensorDesc))
 			{
-				UE_LOG(LogNNX, Warning, TEXT("Failed to initialize tensor(s) for DML inference"));
+				UE_LOG(LogNNE, Warning, TEXT("Failed to initialize tensor(s) for DML inference"));
 				return false;
 			}
 		}
 
 		if (!InitDmlTensorDesc(DmlOutputTensorDesc, OutputTensorDesc))
 		{
-			UE_LOG(LogNNX, Warning, TEXT("Failed to initialize tensor(s) for DML inference"));
+			UE_LOG(LogNNE, Warning, TEXT("Failed to initialize tensor(s) for DML inference"));
 			return false;
 		}
 
@@ -770,266 +766,7 @@ public:
 
 };
 
-//
-//
-//
-class FInferenceModelDml : public FInferenceModelRDG
-{
-	class FGraphBuilder;
-	class FBindingTable;
-	class FDebugName;
-
-public:
-
-	FInferenceModelDml();
-	~FInferenceModelDml();
-
-	bool Init(TConstArrayView<uint8> ModelData, FDeviceContextDml* InDevCtx);
-
-protected:
-
-	virtual void AddDispatchOps_RenderThread(FRDGBuilder& GraphBuilder) override;
-	virtual int PrepareTensorShapesAndData() override;
-
-private:
-
-	bool InitCompiledOp(TConstArrayView<int32> OpInputIndices, uint64 TensorDataSize);
-
-	FOperatorDml* OpCreate(const FString& Name, TArrayView<const NNECore::Internal::FTensor> InputTensorDesc, TArrayView<const NNECore::Internal::FTensor> OutputTensorDescs, const NNECore::FAttributeMap& Attributes);
-		
-	FBufferRHIRef CreateRHIBuffer(FRHICommandListImmediate& RHICmdList, uint32 Size, EBufferUsageFlags Usage, ERHIAccess Access, const TCHAR* DbgName);
-	ID3D12Resource* CreateD3D12Buffer(uint32 Size, D3D12_RESOURCE_STATES ResourceState = D3D12_RESOURCE_STATE_COMMON, D3D12_HEAP_TYPE HeapType = D3D12_HEAP_TYPE_DEFAULT, const TCHAR* DebugName = nullptr);
-	
-	// TODO: FIXME: This should go into RDG
-	static constexpr int32 MaxNumInputs = 32;
-	static constexpr int32 MaxNumOutputs = 4;
-
-	using FRHIBufferInputArray = TArray<FRHIBuffer*, TInlineAllocator<MaxNumInputs>>;
-	using FRHIBufferOutputArray = TArray<FRHIBuffer*, TInlineAllocator<MaxNumOutputs>>;
-
-	TComPtr<IDMLOperatorInitializer>	OpInit;
-	TComPtr<IDMLCompiledOperator>		CompiledOp;
-	FDeviceContextDml*					DevCtx;
-	TUniquePtr<FBindingTable>			BindingTable;
-	TComPtr<ID3D12DescriptorHeap>		DescHeap;
-	uint32								DescCount;
-	uint32								DescSize;
-
-	FRHIBufferInputArray				InputBuffers;
-	FRHIBufferOutputArray				OutputBuffers;
-#ifdef NNE_USE_D3D12_RESOURCES
-	TComPtr<ID3D12Resource>				PersistBuff;
-	TComPtr<ID3D12Resource>				TempBuff;
-#else
-	FBufferRHIRef						PersistBuff;
-#endif
-	uint64								MemSizeTemp;
-	uint64								MemSizePersist;
-	ID3D12DynamicRHI*					DynamicRHI;
-};
-
-//
-//
-//
-class FRuntimeDml : public FRuntimeRDG
-{
-public:
-
-	FRuntimeDml() = default;
-	virtual ~FRuntimeDml();
-
-	virtual FString GetRuntimeName() const override
-	{
-		return NNX_RUNTIME_DML_NAME;
-	}
-	
-	virtual EMLRuntimeSupportFlags GetSupportFlags() const override
-	{
-		return EMLRuntimeSupportFlags::RDG;
-	}
-
-	virtual TArray<uint8> CreateModelData(FString FileType, TConstArrayView<uint8> FileData) override
-	{
-		if (!CanCreateModelData(FileType, FileData))
-		{
-			return {};
-		}
-
-		TUniquePtr<IModelOptimizer> Optimizer = NNEUtils::Internal::CreateONNXToNNEModelOptimizer();
-
-		FNNIModelRaw InputModel;
-		InputModel.Data = FileData;
-		InputModel.Format = ENNXInferenceFormat::ONNX;
-
-		FNNIModelRaw OutputModel;
-		if (!Optimizer->Optimize(InputModel, OutputModel, {}))
-		{
-			return {};
-		}
-
-		return ConvertToModelData(OutputModel.Data);
-	};
-
-	virtual TUniquePtr<FMLInferenceModel> CreateModel(TConstArrayView<uint8> ModelData) override
-	{
-		if (!CanCreateModel(ModelData))
-		{
-			return TUniquePtr<FMLInferenceModel>();
-		}
-
-		// Create the model and initialize it with the data not including the header
-		FInferenceModelDml* Model = new FInferenceModelDml();
-		if (!Model->Init(ModelData, &Ctx))
-		{
-			delete Model;
-			return TUniquePtr<FMLInferenceModel>();
-		}
-		return TUniquePtr<FMLInferenceModel>(Model);
-	}
-
-	bool Init(bool bRegisterOnlyOperators);
-
-private:
-
-	bool RegisterElementWiseUnaryOperators();
-	bool RegisterElementWiseBinaryOperators();
-	bool RegisterGemmOperator();
-	
-	FDeviceContextDml		Ctx;
-};
-
-//
-//
-//
-static TUniquePtr<FRuntimeDml> GDmlRuntime;
-
-//
-//
-//
-FRuntimeDml::~FRuntimeDml()
-{
-}
-
-//
-//
-//
-bool FRuntimeDml::Init(bool bRegisterOnlyOperators)
-{
-	RegisterElementWiseUnaryOperators();
-	RegisterElementWiseBinaryOperators();
-	RegisterGemmOperator();
-
-	if (bRegisterOnlyOperators)
-	{
-		UE_LOG(LogNNX, Display, TEXT("Registering only operators"));
-		return true;
-	}
-
-	HRESULT Res;
-
-	// In order to use DirectML we need D3D12
-	ID3D12DynamicRHI* RHI = nullptr;
-
-	if (GDynamicRHI && GDynamicRHI->GetInterfaceType() == ERHIInterfaceType::D3D12)
-	{
-		RHI = GetID3D12PlatformDynamicRHI();
-
-		if (!RHI)
-		{
-			UE_LOG(LogNNX, Warning, TEXT("Error:%s RHI is not supported by DirectML"), GDynamicRHI->GetName());
-			return false;
-		}
-	}
-	else
-	{
-		if (GDynamicRHI)
-		{
-			UE_LOG(LogNNX, Warning, TEXT("Error:%s RHI is not supported by DirectML"), GDynamicRHI->GetName());
-			return false;
-		}
-		else
-		{
-			UE_LOG(LogNNX, Warning, TEXT("Error:No RHI found"));
-			return false;
-		}
-	}
-
-	Ctx.DeviceIndex = 0;
-	Ctx.D3D12Device = RHI->RHIGetDevice(Ctx.DeviceIndex);
-
-#if PLATFORM_WINDOWS
-	ID3D12Device5* D3D12Device5 = nullptr;
-	
-	Res = Ctx.D3D12Device->QueryInterface(&D3D12Device5);
-	if (D3D12Device5)
-	{
-		uint32	NumCommands = 0;
-
-		Res = D3D12Device5->EnumerateMetaCommands(&NumCommands, nullptr);
-		if (NumCommands)
-		{
-			UE_LOG(LogNNX, Verbose, TEXT("D3D12 Meta commands:%u"), NumCommands);
-
-			TArray<D3D12_META_COMMAND_DESC>	MetaCmds;
-
-			MetaCmds.SetNumUninitialized(NumCommands);
-
-			Res = D3D12Device5->EnumerateMetaCommands(&NumCommands, MetaCmds.GetData());
-			for (uint32 Idx = 0; Idx < NumCommands; ++Idx)
-			{
-				const D3D12_META_COMMAND_DESC& Desc = MetaCmds[Idx];
-
-				UE_LOG(LogNNX, Verbose, TEXT("   %s"), Desc.Name);
-			}
-		}
-	}
-#endif
-
-	DML_CREATE_DEVICE_FLAGS DmlCreateFlags = DML_CREATE_DEVICE_FLAG_NONE;
-
-	// Set debugging flags
-	if (RHI->IsD3DDebugEnabled())
-	{
-		DmlCreateFlags |= DML_CREATE_DEVICE_FLAG_DEBUG;
-	}
-
-	Res = DMLCreateDevice(Ctx.D3D12Device, DmlCreateFlags, DML_PPV_ARGS(&Ctx.Device));
-	if (!Ctx.Device)
-	{
-		UE_LOG(LogNNX, Warning, TEXT("Failed to create DirectML device, res:%x"), Res);
-		return false;
-	}
-
-	DML_FEATURE_QUERY_TENSOR_DATA_TYPE_SUPPORT	Fp16Query = { DML_TENSOR_DATA_TYPE_FLOAT16 };
-	DML_FEATURE_DATA_TENSOR_DATA_TYPE_SUPPORT	Fp16Supported = {};
-	
-	Ctx.Device->CheckFeatureSupport(DML_FEATURE_TENSOR_DATA_TYPE_SUPPORT, sizeof(Fp16Query), &Fp16Query, sizeof(Fp16Supported), &Fp16Supported);
-
-	DML_FEATURE_LEVEL					FeatureLevels[] = { DML_FEATURE_LEVEL_5_0 };
-	DML_FEATURE_QUERY_FEATURE_LEVELS	FeatureLevelQuery = { UE_ARRAY_COUNT(FeatureLevels), FeatureLevels };
-	DML_FEATURE_DATA_FEATURE_LEVELS		FeatureLevelSupported = {};
-
-	Res = Ctx.Device->CheckFeatureSupport(DML_FEATURE_FEATURE_LEVELS, sizeof(FeatureLevelQuery), &FeatureLevelQuery, sizeof(FeatureLevelSupported), &FeatureLevelSupported);
-	if (FAILED(Res) || FeatureLevelSupported.MaxSupportedFeatureLevel < DML_FEATURE_LEVEL_5_0)
-	{
-		UE_LOG(LogNNX, Warning, TEXT("DirectML feature level %x not supported"), FeatureLevels[0]);
-		return false;
-	}
-
-	Res = Ctx.Device->CreateCommandRecorder(DML_PPV_ARGS(&Ctx.CmdRec));
-	if (!Ctx.CmdRec)
-	{
-		UE_LOG(LogNNX, Warning, TEXT("Failed to create DML command recorder, res:%x"), Res);
-		return false;
-	}
-
-	return true;
-}
-
-//
-//
-//
-bool FRuntimeDml::RegisterElementWiseUnaryOperators()
+bool RegisterElementWiseUnaryOperators()
 {
 #define OP(DmlOpDesc, OpName) FOperatorRegistryDml::Get()->OpAdd(TEXT(#OpName), FOperatorDmlElementWiseUnary<DmlOpDesc, NNECore::Internal::EElementWiseUnaryOperatorType::OpName>::Create)
 
@@ -1074,10 +811,7 @@ bool FRuntimeDml::RegisterElementWiseUnaryOperators()
 	return true;
 }
 
-//
-//
-//
-bool FRuntimeDml::RegisterElementWiseBinaryOperators()
+bool RegisterElementWiseBinaryOperators()
 {
 #define OP(DmlOpDesc, OpName) FOperatorRegistryDml::Get()->OpAdd(TEXT(#OpName), FOperatorDmlElementWiseBinary<DmlOpDesc, NNECore::Internal::EElementWiseBinaryOperatorType::OpName>::Create)
 
@@ -1102,10 +836,7 @@ bool FRuntimeDml::RegisterElementWiseBinaryOperators()
 	return true;
 }
 
-//
-//
-//
-bool FRuntimeDml::RegisterGemmOperator()
+bool RegisterGemmOperator()
 {
 	FOperatorRegistryDml::Get()->OpAdd(TEXT("Gemm"), FOperatorDmlGemm::Create);
 	return true;
@@ -1114,7 +845,65 @@ bool FRuntimeDml::RegisterGemmOperator()
 //
 //
 //
-class FInferenceModelDml::FDebugName
+class FModel : public FModelRDG
+{
+	class FGraphBuilder;
+	class FBindingTable;
+	class FDebugName;
+
+public:
+
+	FModel();
+	~FModel();
+
+	bool Init(TConstArrayView<uint8> ModelData, FDeviceContextDml* InDevCtx);
+
+protected:
+
+	virtual void AddDispatchOps_RenderThread(FRDGBuilder& GraphBuilder) override;
+	virtual int PrepareTensorShapesAndData() override;
+
+private:
+
+	bool InitCompiledOp(TConstArrayView<int32> OpInputIndices, uint64 TensorDataSize);
+
+	FOperatorDml* OpCreate(const FString& Name, TArrayView<const NNECore::Internal::FTensor> InputTensorDesc, TArrayView<const NNECore::Internal::FTensor> OutputTensorDescs, const NNECore::FAttributeMap& Attributes);
+
+	FBufferRHIRef CreateRHIBuffer(FRHICommandListImmediate& RHICmdList, uint32 Size, EBufferUsageFlags Usage, ERHIAccess Access, const TCHAR* DbgName);
+	ID3D12Resource* CreateD3D12Buffer(uint32 Size, D3D12_RESOURCE_STATES ResourceState = D3D12_RESOURCE_STATE_COMMON, D3D12_HEAP_TYPE HeapType = D3D12_HEAP_TYPE_DEFAULT, const TCHAR* DebugName = nullptr);
+
+	// TODO: FIXME: This should go into RDG
+	static constexpr int32 MaxNumInputs = 32;
+	static constexpr int32 MaxNumOutputs = 4;
+
+	using FRHIBufferInputArray = TArray<FRHIBuffer*, TInlineAllocator<MaxNumInputs>>;
+	using FRHIBufferOutputArray = TArray<FRHIBuffer*, TInlineAllocator<MaxNumOutputs>>;
+
+	TComPtr<IDMLOperatorInitializer>	OpInit;
+	TComPtr<IDMLCompiledOperator>		CompiledOp;
+	FDeviceContextDml*					DevCtx;
+	TUniquePtr<FBindingTable>			BindingTable;
+	TComPtr<ID3D12DescriptorHeap>		DescHeap;
+	uint32								DescCount;
+	uint32								DescSize;
+
+	FRHIBufferInputArray				InputBuffers;
+	FRHIBufferOutputArray				OutputBuffers;
+#ifdef NNE_USE_D3D12_RESOURCES
+	TComPtr<ID3D12Resource>				PersistBuff;
+	TComPtr<ID3D12Resource>				TempBuff;
+#else
+	FBufferRHIRef						PersistBuff;
+#endif
+	uint64								MemSizeTemp;
+	uint64								MemSizePersist;
+	ID3D12DynamicRHI* DynamicRHI;
+};
+
+//
+//
+//
+class FModel::FDebugName
 {
 	static constexpr int32 Size = 128;
 
@@ -1161,11 +950,11 @@ private:
 //
 //
 //
-class FInferenceModelDml::FBindingTable
+class FModel::FBindingTable
 {
 public:
 
-	bool Init(FInferenceModelDml* InModel)
+	bool Init(FModel* InModel)
 	{
 		Model = InModel;
 		DynamicRHI = InModel->DynamicRHI;
@@ -1308,7 +1097,7 @@ private:
 			Res = Model->DevCtx->Device->CreateBindingTable(&Desc, DML_PPV_ARGS(&BindingTable));
 			if (!BindingTable)
 			{
-				UE_LOG(LogNNX, Warning, TEXT("Failed to create DML binding table, res:%d"), Res);
+				UE_LOG(LogNNE, Warning, TEXT("Failed to create DML binding table, res:%d"), Res);
 				return false;
 			}
 		}
@@ -1348,13 +1137,13 @@ private:
 	TArray<DML_BUFFER_BINDING, TInlineAllocator<MaxNumOutputs>>		OutputBinds;
 	TArray<DML_BINDING_DESC, TInlineAllocator<MaxNumOutputs>>		OutputBindDescs;
 	ID3D12DynamicRHI*												DynamicRHI;
-	FInferenceModelDml*												Model;
+	FModel*															Model;
 };
 
 //
 //
 //
-class FInferenceModelDml::FGraphBuilder
+class FModel::FGraphBuilder
 {
 public:
 
@@ -1573,7 +1362,7 @@ public:
 		Res = Device1->CompileGraph(&Graph, DML_EXECUTION_FLAG_NONE, DML_PPV_ARGS(&Op));
 		if (FAILED(Res))
 		{
-			UE_LOG(LogNNX, Warning, TEXT("Error:Failed to compile DML graph"));
+			UE_LOG(LogNNE, Warning, TEXT("Error:Failed to compile DML graph"));
 			Op = nullptr;
 		};
 
@@ -1775,27 +1564,28 @@ private:
 //
 //
 //
-FInferenceModelDml::FInferenceModelDml()
+FModel::FModel()
 {
-	bUseManualTransitions = true;
 }
 
 //
 //
 //
-FInferenceModelDml::~FInferenceModelDml()
+FModel::~FModel()
 {
 }
 
 //
 //
 //
-bool FInferenceModelDml::Init(TConstArrayView<uint8> ModelData, FDeviceContextDml* InDevCtx)
+bool FModel::Init(TConstArrayView<uint8> ModelData, FDeviceContextDml* InDevCtx)
 {
 	check(ModelData.Num() > 0);
 	FMLRuntimeFormat	Format;
+	int32 GuidSize = sizeof(UNNERuntimeRDGDmlImpl::GUID);
+	int32 VersionSize = sizeof(UNNERuntimeRDGDmlImpl::Version);
 
-	if (!LoadModel(ModelData, Format))
+	if (!LoadModel(ModelData, Format, GuidSize + VersionSize))
 	{
 		return false;
 	}
@@ -1806,7 +1596,7 @@ bool FInferenceModelDml::Init(TConstArrayView<uint8> ModelData, FDeviceContextDm
 	HRESULT Res = DevCtx->Device->CreateOperatorInitializer(0, nullptr, DML_PPV_ARGS(&OpInit));
 	if (!OpInit)
 	{
-		UE_LOG(LogNNX, Warning, TEXT("Error:Failed to create DML operator initializer"));
+		UE_LOG(LogNNE, Warning, TEXT("Error:Failed to create DML operator initializer"));
 		return false;
 	}
 
@@ -1885,7 +1675,7 @@ bool FInferenceModelDml::Init(TConstArrayView<uint8> ModelData, FDeviceContextDm
 
 		if (!OpDesc.Op)
 		{
-			UE_LOG(LogNNX, Warning, TEXT("Error:Failed to create operator:%s"), *TypeName);
+			UE_LOG(LogNNE, Warning, TEXT("Error:Failed to create operator:%s"), *TypeName);
 			return false;
 		}
 
@@ -1912,7 +1702,7 @@ bool FInferenceModelDml::Init(TConstArrayView<uint8> ModelData, FDeviceContextDm
 //
 //
 //
-bool FInferenceModelDml::InitCompiledOp(TConstArrayView<int32> OpInputIndices, uint64 TensorDataSize)
+bool FModel::InitCompiledOp(TConstArrayView<int32> OpInputIndices, uint64 TensorDataSize)
 {
 	static constexpr EBufferUsageFlags	WeightBuffUsage = BUF_UnorderedAccess;
 	static constexpr ERHIAccess			WeightBuffAccess = ERHIAccess::UAVMask;
@@ -1930,7 +1720,7 @@ bool FInferenceModelDml::InitCompiledOp(TConstArrayView<int32> OpInputIndices, u
 	Res = OpInit->Reset(UE_ARRAY_COUNT(CompiledOps), CompiledOps);
 	if (FAILED(Res))
 	{
-		UE_LOG(LogNNX, Warning, TEXT("Error:Failed to reset DirectML operator initializer"));
+		UE_LOG(LogNNE, Warning, TEXT("Error:Failed to reset DirectML operator initializer"));
 		return false;
 	}
 
@@ -1948,7 +1738,7 @@ bool FInferenceModelDml::InitCompiledOp(TConstArrayView<int32> OpInputIndices, u
 	Res = DevCtx->D3D12Device->CreateDescriptorHeap(&HeapDesc, IID_PPV_ARGS(&DescHeap));
 	if (!DescHeap)
 	{
-		UE_LOG(LogNNX, Warning, TEXT("Failed to create descriptor heap, res:%x"), Res);
+		UE_LOG(LogNNE, Warning, TEXT("Failed to create descriptor heap, res:%x"), Res);
 		return false;
 	}
 
@@ -1965,7 +1755,7 @@ bool FInferenceModelDml::InitCompiledOp(TConstArrayView<int32> OpInputIndices, u
 
 	FEvent* Signal = FGenericPlatformProcess::GetSynchEventFromPool(false);
 
-	ENQUEUE_RENDER_COMMAND(FInferenceModelDml_SetTensorData)
+	ENQUEUE_RENDER_COMMAND(FModel_SetTensorData)
 	(
 		[
 			this, 
@@ -1999,7 +1789,7 @@ bool FInferenceModelDml::InitCompiledOp(TConstArrayView<int32> OpInputIndices, u
 
 					FBufferRHIRef WeightBuff;
 
-					WeightBuff = CreateRHIBuffer(RHICmdList, TensorData.Num(), WeightBuffUsage, WeightBuffAccess, TEXT("FInferenceModelDml_TensorWeights"));
+					WeightBuff = CreateRHIBuffer(RHICmdList, TensorData.Num(), WeightBuffUsage, WeightBuffAccess, TEXT("FModel_TensorWeights"));
 				
 					FMemory::Memcpy(UploadBuffPtr + UploadOffset, TensorData.GetData(), TensorData.Num());
 					RHICmdList.CopyBufferRegion(WeightBuff, 0, UploadBuff, UploadOffset, TensorData.Num());
@@ -2097,7 +1887,7 @@ END_SHADER_PARAMETER_STRUCT()
 //
 //
 //
-void FInferenceModelDml::AddDispatchOps_RenderThread(FRDGBuilder& GraphBuilder)
+void FModel::AddDispatchOps_RenderThread(FRDGBuilder& GraphBuilder)
 {
 	const ERDGPassFlags TransitionBuffFlags = ERDGPassFlags::Compute | ERDGPassFlags::NeverCull;
 
@@ -2110,7 +1900,7 @@ void FInferenceModelDml::AddDispatchOps_RenderThread(FRDGBuilder& GraphBuilder)
 		Params->Buffer = AllTensorRDGs[InputTensorIndices[Idx]]->GetBuffer();
 
 		GraphBuilder.AddPass(
-			RDG_EVENT_NAME("FInferenceModelDml_Dispatch_GetInputBuffer"),
+			RDG_EVENT_NAME("FModel_Dispatch_GetInputBuffer"),
 			Params,
 			TransitionBuffFlags,
 			[this, Idx, Params](FRHICommandListImmediate& RHICmdList)
@@ -2131,7 +1921,7 @@ void FInferenceModelDml::AddDispatchOps_RenderThread(FRDGBuilder& GraphBuilder)
 		Params->Buffer = AllTensorRDGs[OutputTensorIndices[Idx]]->GetBuffer();
 
 		GraphBuilder.AddPass(
-			RDG_EVENT_NAME("FInferenceModelDml_Dispatch_GetOutputBuffer"),
+			RDG_EVENT_NAME("FModel_Dispatch_GetOutputBuffer"),
 			Params,
 			TransitionBuffFlags,
 			[this, Idx, Params](FRHICommandListImmediate& RHICmdList)
@@ -2142,7 +1932,7 @@ void FInferenceModelDml::AddDispatchOps_RenderThread(FRDGBuilder& GraphBuilder)
 	}
 	
 	GraphBuilder.AddPass(
-		RDG_EVENT_NAME("FInferenceModelDml_Dispatch"),
+		RDG_EVENT_NAME("FModel_Dispatch"),
 		ERDGPassFlags::None | ERDGPassFlags::NeverCull,
 		[this](FRHICommandListImmediate& RHICmdList)
 		{
@@ -2170,7 +1960,7 @@ void FInferenceModelDml::AddDispatchOps_RenderThread(FRDGBuilder& GraphBuilder)
 					}
 
 					// TODO: We should use this instead of NNE_USE_D3D12_RESOURCES
-					//FBufferRHIRef TempBuff = MemSizeTemp ? CreateRHIBuffer(RHICmdList, MemSizeTemp, TempBuffUsage, TempBuffAccess, TEXT("FInferenceModelDml_Dispatch_TempBuff")) : nullptr;
+					//FBufferRHIRef TempBuff = MemSizeTemp ? CreateRHIBuffer(RHICmdList, MemSizeTemp, TempBuffUsage, TempBuffAccess, TEXT("FModel_Dispatch_TempBuff")) : nullptr;
 
 					BindingTable->Bind(CompiledOp, InputBuffers, OutputBuffers, PersistBuff, TempBuff);
 
@@ -2191,13 +1981,13 @@ void FInferenceModelDml::AddDispatchOps_RenderThread(FRDGBuilder& GraphBuilder)
 //
 // Create operator
 //
-FOperatorDml* FInferenceModelDml::OpCreate(const FString& OpName, TArrayView<const NNECore::Internal::FTensor> InputTensorDescs, TArrayView<const NNECore::Internal::FTensor> OutputTensorDescs, const NNECore::FAttributeMap& Attributes)
+FOperatorDml* FModel::OpCreate(const FString& OpName, TArrayView<const NNECore::Internal::FTensor> InputTensorDescs, TArrayView<const NNECore::Internal::FTensor> OutputTensorDescs, const NNECore::FAttributeMap& Attributes)
 {
 	FOperatorRegistryDml::OperatorCreateFunc CreateFn = FOperatorRegistryDml::Get()->OpFind(OpName);
 
 	if (!CreateFn)
 	{
-		UE_LOG(LogNNX, Warning, TEXT("Dml MLOperatorRegistry failed to find operator:%s"), *OpName);
+		UE_LOG(LogNNE, Warning, TEXT("Dml MLOperatorRegistry failed to find operator:%s"), *OpName);
 		return nullptr;
 	}
 
@@ -2207,7 +1997,7 @@ FOperatorDml* FInferenceModelDml::OpCreate(const FString& OpName, TArrayView<con
 	{
 		delete Op;
 
-		UE_LOG(LogNNX, Warning, TEXT("Error:Failed to initialize operator:%s"), *OpName);
+		UE_LOG(LogNNE, Warning, TEXT("Error:Failed to initialize operator:%s"), *OpName);
 		return nullptr;
 	}
 
@@ -2216,7 +2006,7 @@ FOperatorDml* FInferenceModelDml::OpCreate(const FString& OpName, TArrayView<con
 	return Op;
 }
 
-FBufferRHIRef FInferenceModelDml::CreateRHIBuffer(FRHICommandListImmediate& RHICmdList, uint32 Size, EBufferUsageFlags Usage, ERHIAccess Access, const TCHAR* DbgName)
+FBufferRHIRef FModel::CreateRHIBuffer(FRHICommandListImmediate& RHICmdList, uint32 Size, EBufferUsageFlags Usage, ERHIAccess Access, const TCHAR* DbgName)
 {
 	FBufferRHIRef Buff = nullptr;
 
@@ -2231,7 +2021,7 @@ FBufferRHIRef FInferenceModelDml::CreateRHIBuffer(FRHICommandListImmediate& RHIC
 	return Buff;
 }
 
-ID3D12Resource* FInferenceModelDml::CreateD3D12Buffer(uint32 Size, D3D12_RESOURCE_STATES ResourceState, D3D12_HEAP_TYPE HeapType, const TCHAR* DebugName)
+ID3D12Resource* FModel::CreateD3D12Buffer(uint32 Size, D3D12_RESOURCE_STATES ResourceState, D3D12_HEAP_TYPE HeapType, const TCHAR* DebugName)
 {
 	ID3D12Resource* Resource = nullptr;
 
@@ -2249,7 +2039,7 @@ ID3D12Resource* FInferenceModelDml::CreateD3D12Buffer(uint32 Size, D3D12_RESOURC
 
 	if (FAILED(Res))
 	{
-		UE_LOG(LogNNX, Warning, TEXT("Error:FInferenceModel failed to create D3D12 resource"));
+		UE_LOG(LogNNE, Warning, TEXT("Error:FInferenceModel failed to create D3D12 resource"));
 		return nullptr;
 	}
 
@@ -2264,13 +2054,13 @@ ID3D12Resource* FInferenceModelDml::CreateD3D12Buffer(uint32 Size, D3D12_RESOURC
 //
 //
 //
-int FInferenceModelDml::PrepareTensorShapesAndData()
+int FModel::PrepareTensorShapesAndData()
 {
 	for (FTensorDesc SymbolicTensorDesc : AllSymbolicTensorDescs)
 	{
 		if (!SymbolicTensorDesc.GetShape().IsConcrete())
 		{
-			UE_LOG(LogNNX, Warning, TEXT("DML engine does not support model with variable shapes yet."));
+			UE_LOG(LogNNE, Warning, TEXT("DML engine does not support model with variable shapes yet."));
 			return -1;
 		}
 	}
@@ -2279,91 +2069,268 @@ int FInferenceModelDml::PrepareTensorShapesAndData()
 }
 
 //
+// Called on RDG runtime startup
 //
-//
-static TUniquePtr<FRuntimeDml> FDmlRuntimeCreate(bool bRegisterOnlyOperators)
+bool FRuntimeDmlStartup()
 {
-	auto Runtime = MakeUnique<FRuntimeDml>();
+	bool bIsD3D12RHI = GDynamicRHI && GDynamicRHI->GetInterfaceType() == ERHIInterfaceType::D3D12;		
+	bool bLoadDirectML = true;
 
-	if (!Runtime->Init(bRegisterOnlyOperators))
+	if (IsRunningCommandlet() && !IsAllowCommandletRendering())
 	{
-		UE_LOG(LogNNX, Warning, TEXT("Failed to create NNX DML runtime"));
-		Runtime.Release();
+		UE_LOG(LogNNE, Display, TEXT("Running inside commandlet without rendering"));
+		bLoadDirectML = false;
 	}
 
-	return Runtime;
+#ifdef DIRECTML_BIN_PATH
+		
+	if (bIsD3D12RHI && bLoadDirectML)
+	{
+		const FString DirectMLRuntimeBinPath = TEXT(PREPROCESSOR_TO_STRING(DIRECTML_BIN_PATH));
+		FString DirectMLDLLPaths[2];
+		int32	NumPaths = 1;
+		
+		DirectMLDLLPaths[0] = DirectMLRuntimeBinPath / TEXT("DirectML.dll");
+
+		if (GetID3D12PlatformDynamicRHI()->IsD3DDebugEnabled())
+		{
+			DirectMLDLLPaths[1] = DirectMLRuntimeBinPath / TEXT("DirectML.Debug.dll");
+			++NumPaths;
+		}
+
+		FPlatformProcess::PushDllDirectory(*DirectMLRuntimeBinPath);
+
+		for (int32 Idx = 0; Idx < NumPaths; ++Idx)
+		{
+			if (!FPaths::FileExists(DirectMLDLLPaths[Idx]))
+			{
+				const FString ErrorMessage = FString::Format(TEXT("DirectML DLL file not found in \"{0}\"."),
+					{ IFileManager::Get().ConvertToAbsolutePathForExternalAppForRead(*DirectMLDLLPaths[Idx])});
+				UE_LOG(LogNNE, Warning, TEXT("NNXRuntimeDll:%s"), *ErrorMessage);
+				checkf(false, TEXT("%s"), *ErrorMessage);
+			}
+
+			FPlatformProcess::GetDllHandle(*DirectMLDLLPaths[Idx]);
+		}
+
+		FPlatformProcess::PopDllDirectory(*DirectMLRuntimeBinPath);
+	}
+#endif
+
+	const bool bRegisterOnlyOperators = !bLoadDirectML;
+	return bRegisterOnlyOperators;
 }
 
-//
-// Called on RDG runtime startup
-//
-IRuntime* FRuntimeDmlStartup()
+} // namespace UE::NNERuntimeRDG::Private::Dml
+
+#endif
+
+using namespace UE::NNERuntimeRDG::Private::Dml;
+
+FGuid UNNERuntimeRDGDmlImpl::GUID = FGuid((int32)'R', (int32)'D', (int32)'G', (int32)'D');
+int32 UNNERuntimeRDGDmlImpl::Version = 0x00000001;
+
+bool UNNERuntimeRDGDmlImpl::CanCreateModelData(FString FileType, TConstArrayView<uint8> FileData) const
+{
+	return FileType.Compare("onnx", ESearchCase::IgnoreCase) == 0;
+}
+
+TArray<uint8> UNNERuntimeRDGDmlImpl::CreateModelData(FString FileType, TConstArrayView<uint8> FileData)
+{
+	if (!CanCreateModelData(FileType, FileData))
+	{
+		return {};
+	}
+
+	TUniquePtr<IModelOptimizer> Optimizer = UE::NNEUtils::Internal::CreateONNXToNNEModelOptimizer();
+
+	FNNIModelRaw InputModel;
+	InputModel.Data = FileData;
+	InputModel.Format = ENNXInferenceFormat::ONNX;
+
+	FNNIModelRaw OutputModel;
+	if (!Optimizer->Optimize(InputModel, OutputModel, {}))
+	{
+		return {};
+	}
+
+	TArray<uint8> Result;
+	FMemoryWriter Writer(Result);
+
+	Writer << GUID;
+	Writer << Version;
+	Writer.Serialize(OutputModel.Data.GetData(), OutputModel.Data.Num());
+	return Result;
+};
+
+bool UNNERuntimeRDGDmlImpl::CanCreateModelRDG(TObjectPtr<UNNEModelData> ModelData) const
 {
-	if (!GDmlRuntime)
+#ifdef NNE_USE_DIRECTML
+	int32 GuidSize = sizeof(GUID);
+	int32 VersionSize = sizeof(Version);
+	TConstArrayView<uint8> Data = ModelData->GetModelData(GetRuntimeName());
+
+	if (Data.Num() <= GuidSize + VersionSize)
 	{
-		bool bIsD3D12RHI = GDynamicRHI && GDynamicRHI->GetInterfaceType() == ERHIInterfaceType::D3D12;		
-		bool bLoadDirectML = true;
+		return false;
+	}
+	bool bResult = FGenericPlatformMemory::Memcmp(&(Data[0]), &(GUID), GuidSize) == 0;
+	bResult &= FGenericPlatformMemory::Memcmp(&(Data[GuidSize]), &(Version), VersionSize) == 0;
+	return bResult;
+#else
+	return false;
+#endif
+};
 
-		if (IsRunningCommandlet() && !IsAllowCommandletRendering())
+TUniquePtr<UE::NNECore::IModelRDG> UNNERuntimeRDGDmlImpl::CreateModelRDG(TObjectPtr<UNNEModelData> ModelData)
+{
+#ifdef NNE_USE_DIRECTML
+	if (!CanCreateModelRDG(ModelData))
+	{
+		return TUniquePtr<UE::NNECore::IModelRDG>();
+	}
+
+	// Create the model and initialize it with the data not including the header
+	FModel* Model = new FModel();
+	TConstArrayView<uint8> Data = ModelData->GetModelData(GetRuntimeName());
+
+	if (!Model->Init(Data, Ctx))
+	{
+		delete Model;
+		return TUniquePtr<UE::NNECore::IModelRDG>();
+	}
+	return TUniquePtr<UE::NNECore::IModelRDG>(Model);
+#else
+	return TUniquePtr<UE::NNECore::IModelRDG>();
+#endif
+}
+
+UNNERuntimeRDGDmlImpl::~UNNERuntimeRDGDmlImpl()
+{
+#ifdef NNE_USE_DIRECTML
+	if (Ctx)
+	{
+		delete Ctx;
+		Ctx = nullptr;
+	}
+#endif
+}
+
+bool UNNERuntimeRDGDmlImpl::Init(bool bRegisterOnlyOperators)
+{
+#ifdef NNE_USE_DIRECTML
+	RegisterElementWiseUnaryOperators();
+	RegisterElementWiseBinaryOperators();
+	RegisterGemmOperator();
+
+	if (bRegisterOnlyOperators)
+	{
+		UE_LOG(LogNNE, Display, TEXT("Registering only operators"));
+		return true;
+	}
+
+	HRESULT Res;
+
+	// In order to use DirectML we need D3D12
+	ID3D12DynamicRHI* RHI = nullptr;
+
+	if (GDynamicRHI && GDynamicRHI->GetInterfaceType() == ERHIInterfaceType::D3D12)
+	{
+		RHI = GetID3D12PlatformDynamicRHI();
+
+		if (!RHI)
+		{
+			UE_LOG(LogNNE, Warning, TEXT("Error:%s RHI is not supported by DirectML"), GDynamicRHI->GetName());
+			return false;
+		}
+	}
+	else
+	{
+		if (GDynamicRHI)
 		{
-			UE_LOG(LogNNX, Display, TEXT("Running inside commandlet without rendering"));
-			bLoadDirectML = false;
+			UE_LOG(LogNNE, Warning, TEXT("Error:%s RHI is not supported by DirectML"), GDynamicRHI->GetName());
+			return false;
 		}
+		else
+		{
+			UE_LOG(LogNNE, Warning, TEXT("Error:No RHI found"));
+			return false;
+		}
+	}
 
-#ifdef DIRECTML_BIN_PATH
-		
-		if (bIsD3D12RHI && bLoadDirectML)
+	check(Ctx == nullptr);
+	Ctx = new FDeviceContextDml();
+	Ctx->DeviceIndex = 0;
+	Ctx->D3D12Device = RHI->RHIGetDevice(Ctx->DeviceIndex);
+
+#if PLATFORM_WINDOWS
+	ID3D12Device5* D3D12Device5 = nullptr;
+
+	Res = Ctx->D3D12Device->QueryInterface(&D3D12Device5);
+	if (D3D12Device5)
+	{
+		uint32	NumCommands = 0;
+
+		Res = D3D12Device5->EnumerateMetaCommands(&NumCommands, nullptr);
+		if (NumCommands)
 		{
-			const FString DirectMLRuntimeBinPath = TEXT(PREPROCESSOR_TO_STRING(DIRECTML_BIN_PATH));
-			FString DirectMLDLLPaths[2];
-			int32	NumPaths = 1;
-		
-			DirectMLDLLPaths[0] = DirectMLRuntimeBinPath / TEXT("DirectML.dll");
+			UE_LOG(LogNNE, Verbose, TEXT("D3D12 Meta commands:%u"), NumCommands);
 
-			if (GetID3D12PlatformDynamicRHI()->IsD3DDebugEnabled())
-			{
-				DirectMLDLLPaths[1] = DirectMLRuntimeBinPath / TEXT("DirectML.Debug.dll");
-				++NumPaths;
-			}
+			TArray<D3D12_META_COMMAND_DESC>	MetaCmds;
 
-			FPlatformProcess::PushDllDirectory(*DirectMLRuntimeBinPath);
+			MetaCmds.SetNumUninitialized(NumCommands);
 
-			for (int32 Idx = 0; Idx < NumPaths; ++Idx)
+			Res = D3D12Device5->EnumerateMetaCommands(&NumCommands, MetaCmds.GetData());
+			for (uint32 Idx = 0; Idx < NumCommands; ++Idx)
 			{
-				if (!FPaths::FileExists(DirectMLDLLPaths[Idx]))
-				{
-					const FString ErrorMessage = FString::Format(TEXT("DirectML DLL file not found in \"{0}\"."),
-						{ IFileManager::Get().ConvertToAbsolutePathForExternalAppForRead(*DirectMLDLLPaths[Idx])});
-					UE_LOG(LogNNX, Warning, TEXT("NNXRuntimeDll:%s"), *ErrorMessage);
-					checkf(false, TEXT("%s"), *ErrorMessage);
-				}
+				const D3D12_META_COMMAND_DESC& Desc = MetaCmds[Idx];
 
-				FPlatformProcess::GetDllHandle(*DirectMLDLLPaths[Idx]);
+				UE_LOG(LogNNE, Verbose, TEXT("   %s"), Desc.Name);
 			}
-
-			FPlatformProcess::PopDllDirectory(*DirectMLRuntimeBinPath);
 		}
+	}
 #endif
 
-		const bool bRegisterOnlyOperators = !bLoadDirectML;
+	DML_CREATE_DEVICE_FLAGS DmlCreateFlags = DML_CREATE_DEVICE_FLAG_NONE;
 
-		GDmlRuntime = FDmlRuntimeCreate(bRegisterOnlyOperators);
+	// Set debugging flags
+	if (RHI->IsD3DDebugEnabled())
+	{
+		DmlCreateFlags |= DML_CREATE_DEVICE_FLAG_DEBUG;
 	}
 
-	return GDmlRuntime.Get();
-}
+	Res = DMLCreateDevice(Ctx->D3D12Device, DmlCreateFlags, DML_PPV_ARGS(&(Ctx->Device)));
+	if (!Ctx->Device)
+	{
+		UE_LOG(LogNNE, Warning, TEXT("Failed to create DirectML device, res:%x"), Res);
+		return false;
+	}
 
-//
-// Called on RDG runtime shutdown
-//
-void FRuntimeDmlShutdown()
-{
-	if (GDmlRuntime)
+	DML_FEATURE_QUERY_TENSOR_DATA_TYPE_SUPPORT	Fp16Query = { DML_TENSOR_DATA_TYPE_FLOAT16 };
+	DML_FEATURE_DATA_TENSOR_DATA_TYPE_SUPPORT	Fp16Supported = {};
+
+	Ctx->Device->CheckFeatureSupport(DML_FEATURE_TENSOR_DATA_TYPE_SUPPORT, sizeof(Fp16Query), &Fp16Query, sizeof(Fp16Supported), &Fp16Supported);
+
+	DML_FEATURE_LEVEL					FeatureLevels[] = { DML_FEATURE_LEVEL_5_0 };
+	DML_FEATURE_QUERY_FEATURE_LEVELS	FeatureLevelQuery = { UE_ARRAY_COUNT(FeatureLevels), FeatureLevels };
+	DML_FEATURE_DATA_FEATURE_LEVELS		FeatureLevelSupported = {};
+
+	Res = Ctx->Device->CheckFeatureSupport(DML_FEATURE_FEATURE_LEVELS, sizeof(FeatureLevelQuery), &FeatureLevelQuery, sizeof(FeatureLevelSupported), &FeatureLevelSupported);
+	if (FAILED(Res) || FeatureLevelSupported.MaxSupportedFeatureLevel < DML_FEATURE_LEVEL_5_0)
 	{
-		GDmlRuntime.Release();
+		UE_LOG(LogNNE, Warning, TEXT("DirectML feature level %x not supported"), FeatureLevels[0]);
+		return false;
 	}
-}
 
-} // namespace UE::NNERuntimeRDG::Private::Dml
- 
+	Res = Ctx->Device->CreateCommandRecorder(DML_PPV_ARGS(&Ctx->CmdRec));
+	if (!Ctx->CmdRec)
+	{
+		UE_LOG(LogNNE, Warning, TEXT("Failed to create DML command recorder, res:%x"), Res);
+		return false;
+	}
+
+	return true;
+#else
+	return false;
 #endif
+}
diff --git a/Engine/Plugins/Experimental/NNE/Source/NNERuntimeRDG/Private/NNERuntimeRDGDml.h b/Engine/Plugins/Experimental/NNE/Source/NNERuntimeRDG/Private/NNERuntimeRDGDml.h
new file mode 100644
--- /dev/null
+++ b/Engine/Plugins/Experimental/NNE/Source/NNERuntimeRDG/Private/NNERuntimeRDGDml.h
@@ -0,0 +1,47 @@
+// Copyright Epic Games, Inc. All Rights Reserved.
+
+#pragma once
+
+#include "CoreMinimal.h"
+#include "NNECoreRuntime.h"
+#include "NNECoreRuntimeRDG.h"
+#include "Templates/UniquePtr.h"
+#include "UObject/Class.h"
+#include "UObject/Object.h"
+#include "UObject/UObjectBaseUtility.h"
+
+#include "NNERuntimeRDGDml.generated.h"
+
+namespace UE::NNERuntimeRDG::Private::Dml
+{
+	class FDeviceContextDml;
+	bool FRuntimeDmlStartup();
+}
+
+
+UCLASS()
+class UNNERuntimeRDGDmlImpl : public UObject, public INNERuntime, public INNERuntimeRDG
+{
+	GENERATED_BODY()
+
+public:
+	static FGuid GUID;
+	static int32 Version;
+
+	UNNERuntimeRDGDmlImpl():Ctx(nullptr) {};
+	virtual ~UNNERuntimeRDGDmlImpl();
+
+	bool Init(bool bRegisterOnlyOperators);
+
+	virtual FString GetRuntimeName() const override { return TEXT("NNERuntimeDml"); };
+	virtual bool IsPlatformSupported(const ITargetPlatform* TargetPlatform) const override { return true; };
+
+	virtual bool CanCreateModelData(FString FileType, TConstArrayView<uint8> FileData) const override;
+	virtual TArray<uint8> CreateModelData(FString FileType, TConstArrayView<uint8> FileData) override;
+
+	virtual bool CanCreateModelRDG(TObjectPtr<UNNEModelData> ModelData) const override;
+	virtual TUniquePtr<UE::NNECore::IModelRDG> CreateModelRDG(TObjectPtr<UNNEModelData> ModelData) override;
+private:
+	UE::NNERuntimeRDG::Private::Dml::FDeviceContextDml* Ctx;
+};
+
diff --git a/Engine/Plugins/Experimental/NNE/Source/NNERuntimeRDG/Private/NNERuntimeRDGHlsl.cpp b/Engine/Plugins/Experimental/NNE/Source/NNERuntimeRDG/Private/NNERuntimeRDGHlsl.cpp
--- a/Engine/Plugins/Experimental/NNE/Source/NNERuntimeRDG/Private/NNERuntimeRDGHlsl.cpp
+++ b/Engine/Plugins/Experimental/NNE/Source/NNERuntimeRDG/Private/NNERuntimeRDGHlsl.cpp
@@ -1,5 +1,13 @@
 // Copyright Epic Games, Inc. All Rights Reserved.
 
+#include "NNERuntimeRDGHlsl.h"
+#include "NNECoreAttributeMap.h"
+#include "NNECoreModelData.h"
+#include "NNERuntimeRDG.h"
+#include "NNERuntimeRDGHlslHelper.h"
+#include "NNERuntimeRDGModelHlsl.h"
+#include "NNEUtilsModelOptimizer.h"
+#include "NNXModelOptimizerInterface.h"
 #include "Hlsl/NNERuntimeRDGConv.h"
 #include "Hlsl/NNERuntimeRDGConv.h"
 #include "Hlsl/NNERuntimeRDGConvTranspose.h"
@@ -11,144 +19,94 @@
 #include "Hlsl/NNERuntimeRDGPad.h"
 #include "Hlsl/NNERuntimeRDGUpsample.h"
 #include "Hlsl/NNERuntimeRDGMatMul.h"
-#include "NNECoreAttributeMap.h"
-#include "NNERuntimeRDGModelHlsl.h"
-#include "NNXInferenceModel.h"
-#include "NNEUtilsModelOptimizer.h"
-#include "NNXRuntimeFormat.h"
-#include "NNERuntimeRDGHlslHelper.h"
-#include "NNERuntimeRDG.h"
 
-namespace UE::NNERuntimeRDG::Private::Hlsl
-{
+using namespace UE::NNERuntimeRDG::Private::Hlsl;
 
-#define NNX_RUNTIME_HLSL_NAME TEXT("NNXRuntimeHlsl")
+FGuid UNNERuntimeRDGHlslImpl::GUID = FGuid((int32)'R', (int32)'D', (int32)'G', (int32)'H');
+int32 UNNERuntimeRDGHlslImpl::Version = 0x00000001;
 
-//
-//
-//
-class FRuntimeHlsl : public FRuntimeRDG
+bool UNNERuntimeRDGHlslImpl::Init()
 {
-public:
-
-	FRuntimeHlsl() = default;
-
-	virtual ~FRuntimeHlsl()
-	{
-	}
-
-	bool Init()
-	{
-		FOperatorRegistryHlsl* registry = FOperatorRegistryHlsl::Get();
-		check(registry != nullptr);
-
-		RegisterElementWiseUnaryOperators(*registry);
-		RegisterElementWiseBinaryOperators(*registry);
-		RegisterElementWiseVariadicOperators(*registry);
-		RegisterGemmOperator(*registry);
-		RegisterConvOperator(*registry);
-		RegisterConvTransposeOperator(*registry);
-		RegisterMatMulOperator(*registry);
-		RegisterInstanceNormalizationOperator(*registry);
-		RegisterUpsampleOperator(*registry);
-		RegisterPadOperator(*registry);
-
-		return true;
-	}
+	FOperatorRegistryHlsl* registry = FOperatorRegistryHlsl::Get();
+	check(registry != nullptr);
+
+	RegisterElementWiseUnaryOperators(*registry);
+	RegisterElementWiseBinaryOperators(*registry);
+	RegisterElementWiseVariadicOperators(*registry);
+	RegisterGemmOperator(*registry);
+	RegisterConvOperator(*registry);
+	RegisterConvTransposeOperator(*registry);
+	RegisterMatMulOperator(*registry);
+	RegisterInstanceNormalizationOperator(*registry);
+	RegisterUpsampleOperator(*registry);
+	RegisterPadOperator(*registry);
+
+	return true;
+}
 
-	virtual FString GetRuntimeName() const override
-	{
-		return NNX_RUNTIME_HLSL_NAME;
-	}
+bool UNNERuntimeRDGHlslImpl::CanCreateModelData(FString FileType, TConstArrayView<uint8> FileData) const
+{
+	return FileType.Compare("onnx", ESearchCase::IgnoreCase) == 0;
+}
 
-	// Returns flags from ERuntimeSupportFlags
-	virtual NNX::EMLRuntimeSupportFlags GetSupportFlags() const override
+bool UNNERuntimeRDGHlslImpl::CanCreateModelRDG(TObjectPtr<UNNEModelData> ModelData) const
+{
+	int32 GuidSize = sizeof(GUID);
+	int32 VersionSize = sizeof(Version);
+	TConstArrayView<uint8> Data = ModelData->GetModelData(GetRuntimeName());
+	
+	if (Data.Num() <= GuidSize + VersionSize)
 	{
-		return NNX::EMLRuntimeSupportFlags::RDG;
+		return false;
 	}
+	bool bResult = FGenericPlatformMemory::Memcmp(&(Data[0]), &(GUID), GuidSize) == 0;
+	bResult &= FGenericPlatformMemory::Memcmp(&(Data[GuidSize]), &(Version), VersionSize) == 0;
+	return bResult;
+};
 
-	virtual TArray<uint8> CreateModelData(FString FileType, TConstArrayView<uint8> FileData) override
-	{
-		if (!CanCreateModelData(FileType, FileData))
-		{
-			return {};
-		}
-
-		TUniquePtr<NNX::IModelOptimizer> Optimizer = NNEUtils::Internal::CreateONNXToNNEModelOptimizer();
-		Optimizer->AddValidator(MakeShared<FModelValidatorHlsl>());
-
-		FNNIModelRaw InputModel;
-		InputModel.Data = FileData;
-		InputModel.Format = ENNXInferenceFormat::ONNX;
-
-		FNNIModelRaw OutputModel;
-		if (!Optimizer->Optimize(InputModel, OutputModel, {}))
-		{
-			return {};
-		}
-
-		return ConvertToModelData(OutputModel.Data);
-	};
-
-	virtual TUniquePtr<NNX::FMLInferenceModel> CreateModel(TConstArrayView<uint8> ModelData) override
+TArray<uint8> UNNERuntimeRDGHlslImpl::CreateModelData(FString FileType, TConstArrayView<uint8> FileData)
+{
+	if (!CanCreateModelData(FileType, FileData))
 	{
-		if (!CanCreateModel(ModelData))
-		{
-			return TUniquePtr<NNX::FMLInferenceModel>();
-		}
-
-		// Create the model and initialize it with the data not including the header
-		FModel* Model = new FModel();
-		if (!Model->Init(ModelData))
-		{
-			delete Model;
-			return TUniquePtr<NNX::FMLInferenceModel>();
-		}
-		return TUniquePtr<NNX::FMLInferenceModel>(Model);
+		return {};
 	}
-};
 
-/** Globally accessible runtime */
-static TUniquePtr<NNERuntimeRDG::Private::Hlsl::FRuntimeHlsl> GHlslRuntime;
+	TUniquePtr<NNX::IModelOptimizer> Optimizer = UE::NNEUtils::Internal::CreateONNXToNNEModelOptimizer();
+	Optimizer->AddValidator(MakeShared<FModelValidatorHlsl>());
 
-//
-//
-//
-static TUniquePtr<NNERuntimeRDG::Private::Hlsl::FRuntimeHlsl> FRuntimeHlslCreate()
-{
-	auto Runtime = MakeUnique<NNERuntimeRDG::Private::Hlsl::FRuntimeHlsl>();
+	FNNIModelRaw InputModel;
+	InputModel.Data = FileData;
+	InputModel.Format = ENNXInferenceFormat::ONNX;
 
-	if (!Runtime->Init())
+	FNNIModelRaw OutputModel;
+	if (!Optimizer->Optimize(InputModel, OutputModel, {}))
 	{
-		UE_LOG(LogNNX, Warning, TEXT("Failed to create NNX HLSL runtime"));
-		Runtime.Release();
+		return {};
 	}
 
-	return Runtime;
-}
+	TArray<uint8> Result;
+	FMemoryWriter Writer(Result);
+	
+	Writer << GUID;
+	Writer << Version;
+	Writer.Serialize(OutputModel.Data.GetData(), OutputModel.Data.Num());
+	return Result;
+};
 
-//
-// Called on RDG runtime startup
-//
-IRuntime* FRuntimeHlslStartup()
+TUniquePtr<UE::NNECore::IModelRDG> UNNERuntimeRDGHlslImpl::CreateModelRDG(TObjectPtr<UNNEModelData> ModelData)
 {
-	if (!GHlslRuntime)
+	if (!CanCreateModelRDG(ModelData))
 	{
-		GHlslRuntime = FRuntimeHlslCreate();
+		return TUniquePtr<UE::NNECore::IModelRDG>();
 	}
 
-	return GHlslRuntime.Get();
-}
-
-//
-// Called on RDG runtime shutdown
-//
-void FRuntimeHlslShutdown()
-{
-	if (GHlslRuntime)
+	UE::NNERuntimeRDG::Private::Hlsl::FModel* Model = new UE::NNERuntimeRDG::Private::Hlsl::FModel();
+	TConstArrayView<uint8> Data = ModelData->GetModelData(GetRuntimeName());
+	
+	if (!Model->Init(Data))
 	{
-		GHlslRuntime.Release();
+		delete Model;
+		return TUniquePtr<UE::NNECore::IModelRDG>();
 	}
+	return TUniquePtr<UE::NNECore::IModelRDG>(Model);
 }
-
-} // namespace UE::NNERuntimeRDG::Private::Hlsl
diff --git a/Engine/Plugins/Experimental/NNE/Source/NNERuntimeRDG/Private/NNERuntimeRDGHlsl.h b/Engine/Plugins/Experimental/NNE/Source/NNERuntimeRDG/Private/NNERuntimeRDGHlsl.h
new file mode 100644
--- /dev/null
+++ b/Engine/Plugins/Experimental/NNE/Source/NNERuntimeRDG/Private/NNERuntimeRDGHlsl.h
@@ -0,0 +1,37 @@
+// Copyright Epic Games, Inc. All Rights Reserved.
+
+#pragma once
+
+#include "CoreMinimal.h"
+#include "NNECoreRuntime.h"
+#include "NNECoreRuntimeRDG.h"
+#include "Templates/UniquePtr.h"
+#include "UObject/Class.h"
+#include "UObject/Object.h"
+#include "UObject/UObjectBaseUtility.h"
+
+#include "NNERuntimeRDGHlsl.generated.h"
+
+UCLASS()
+class UNNERuntimeRDGHlslImpl : public UObject, public INNERuntime, public INNERuntimeRDG
+{
+	GENERATED_BODY()
+
+public:
+	static FGuid GUID;
+	static int32 Version;
+
+	UNNERuntimeRDGHlslImpl() {};
+	virtual ~UNNERuntimeRDGHlslImpl() {}
+
+	bool Init();
+
+	virtual FString GetRuntimeName() const override { return TEXT("NNERuntimeHlsl"); };
+	virtual bool IsPlatformSupported(const ITargetPlatform* TargetPlatform) const override { return true; };
+
+	virtual bool CanCreateModelData(FString FileType, TConstArrayView<uint8> FileData) const override;
+	virtual TArray<uint8> CreateModelData(FString FileType, TConstArrayView<uint8> FileData) override;
+
+	virtual bool CanCreateModelRDG(TObjectPtr<UNNEModelData> ModelData) const override;
+	virtual TUniquePtr<UE::NNECore::IModelRDG> CreateModelRDG(TObjectPtr<UNNEModelData> ModelData) override;
+};
\ No newline at end of file
diff --git a/Engine/Plugins/Experimental/NNE/Source/NNERuntimeRDG/Private/NNERuntimeRDGModel.cpp b/Engine/Plugins/Experimental/NNE/Source/NNERuntimeRDG/Private/NNERuntimeRDGModel.cpp
new file mode 100644
--- /dev/null
+++ b/Engine/Plugins/Experimental/NNE/Source/NNERuntimeRDG/Private/NNERuntimeRDGModel.cpp
@@ -0,0 +1,267 @@
+// Copyright Epic Games, Inc. All Rights Reserved.
+
+#include "NNERuntimeRDGModel.h"
+#include "NNXInferenceModel.h"
+#include "NNXRuntimeFormat.h"
+#include "NNEUtilsModelOptimizer.h"
+
+#include "RenderGraphBuilder.h"
+#include "RenderGraphUtils.h"
+#include "RHIGPUReadback.h"
+
+#include "Serialization/MemoryReader.h"
+
+namespace UE::NNERuntimeRDG::Private
+{
+
+bool FModelRDG::LoadModel(TConstArrayView<uint8> ModelData, FMLRuntimeFormat& Format, int32 GuidAndVersionSize)
+{
+	TConstArrayView<uint8> ModelBuffer = { &(ModelData.GetData()[GuidAndVersionSize]), ModelData.Num() - GuidAndVersionSize };
+
+	FMemoryReaderView Reader(ModelBuffer);
+
+	FMLRuntimeFormat::StaticStruct()->SerializeBin(Reader, &Format);
+
+	// Data for base class
+	InputSymbolicTensors.Empty();
+	OutputSymbolicTensors.Empty();
+
+	// Data for RDG
+	AllSymbolicTensorDescs.Empty();
+	IntermediateTensorIndices.Empty();
+	WeightTensorIndices.Empty();
+	InputTensorIndices.Empty();
+	OutputTensorIndices.Empty();
+	OperatorInputTensorIndices.Empty();
+	OperatorOutputTensorIndices.Empty();
+
+	// Add tensors
+	for (int32 Idx = 0; Idx < Format.Tensors.Num(); ++Idx)
+	{
+		const FMLFormatTensorDesc& FormatTensorDesc = Format.Tensors[Idx];
+
+		const NNECore::FSymbolicTensorShape SymbolicShape = NNECore::FSymbolicTensorShape::Make(FormatTensorDesc.Shape);
+		const NNECore::FTensorDesc SymbolicTensor = NNECore::FTensorDesc::Make(FormatTensorDesc.Name, SymbolicShape, FormatTensorDesc.DataType);
+
+		AllSymbolicTensorDescs.Emplace(SymbolicTensor);
+
+		if (FormatTensorDesc.Type == EMLFormatTensorType::Input)
+		{
+			InputTensorIndices.Emplace(Idx);
+			InputSymbolicTensors.Emplace(SymbolicTensor);
+		}
+		else if (FormatTensorDesc.Type == EMLFormatTensorType::Output)
+		{
+			OutputTensorIndices.Emplace(Idx);
+			OutputSymbolicTensors.Emplace(SymbolicTensor);
+		}
+		else if (FormatTensorDesc.Type == EMLFormatTensorType::Intermediate)
+		{
+			IntermediateTensorIndices.Emplace(Idx);
+		}
+		else if (FormatTensorDesc.Type == EMLFormatTensorType::Initializer)
+		{
+			WeightTensorIndices.Emplace(Idx);
+			if (!SymbolicTensor.GetShape().IsConcrete())
+			{
+				UE_LOG(LogNNE, Error, TEXT("Weight tensor %s should have a concrete shape"), *SymbolicTensor.GetName());
+				return false;
+			}
+
+			const NNECore::FTensorShape TensorShape = NNECore::FTensorShape::MakeFromSymbolic(SymbolicTensor.GetShape());
+			FTensorRDG& WeightRDG = WeightTensorRDGs.Emplace_GetRef(FTensorRDG::Make(SymbolicTensor, TensorShape, nullptr));
+
+			if (WeightRDG.GetDataSize() != FormatTensorDesc.DataSize)
+			{
+				UE_LOG(LogNNE, Error, TEXT("Weight %s has incorrect size. Expected %d bytes, got %d"), *SymbolicTensor.GetName(), FormatTensorDesc.DataSize, WeightRDG.GetDataSize());
+				return false;
+			}
+
+			const uint8* DataPtr = Format.TensorData.GetData() + FormatTensorDesc.DataOffset;
+			TConstArrayView<uint8> DataView = MakeArrayView(DataPtr, FormatTensorDesc.DataSize);
+
+			WeightRDG.SetPreparedData(DataView);
+		}
+		checkf(FormatTensorDesc.Type != EMLFormatTensorType::None, TEXT("Unsupported tensor type None"));
+	}
+
+	// Loop over all operators in the model and store tensor indices for input/output
+	for (int32 Idx = 0; Idx < Format.Operators.Num(); ++Idx)
+	{
+		OperatorInputTensorIndices.Emplace(Format.Operators[Idx].InTensors);
+		OperatorOutputTensorIndices.Emplace(Format.Operators[Idx].OutTensors);
+	}
+
+	return true;
+}
+
+
+
+int FModelRDG::SetInputTensorShapes(TConstArrayView<FTensorShape> InInputShapes)
+{
+	OutputTensorShapes.Empty();
+
+	//Verify input shape are valid for the model and set InputTensorShapes
+	if (FModelBase<NNECore::IModelRDG>::SetInputTensorShapes(InInputShapes) != 0)
+	{
+		return -1;
+	}
+
+	//Allocate and prime all AllTensorRDGs with concrete shapes defaulting variables dimension to 1 if needed
+	AllTensorRDGs.Init(nullptr, AllSymbolicTensorDescs.Num());
+
+	InputTensorRDGs.Empty();
+	for (int32 i = 0; i < InputTensorIndices.Num(); ++i)
+	{
+		const int32 Idx = InputTensorIndices[i];
+		const FTensorDesc& TensorDesc = InputSymbolicTensors[i];
+		const FTensorShape& TensorShape = InputTensorShapes[i];
+
+		InputTensorRDGs.Emplace(FTensorRDG::Make(TensorDesc, TensorShape, nullptr));
+		AllTensorRDGs[Idx] = &InputTensorRDGs[i];
+	}
+
+	for (int32 i = 0; i < WeightTensorIndices.Num(); ++i)
+	{
+		const int32 Idx = WeightTensorIndices[i];
+
+		AllTensorRDGs[Idx] = &WeightTensorRDGs[i];
+	}
+
+	IntermediateTensorRDGs.Empty();
+	for (int32 i = 0; i < IntermediateTensorIndices.Num(); ++i)
+	{
+		const int32 Idx = IntermediateTensorIndices[i];
+		const FTensorDesc& TensorDesc = AllSymbolicTensorDescs[Idx];
+		const FTensorShape TensorShape = FTensorShape::MakeFromSymbolic(TensorDesc.GetShape());
+
+		IntermediateTensorRDGs.Emplace(FTensorRDG::Make(TensorDesc, TensorShape, nullptr));
+		AllTensorRDGs[Idx] = &IntermediateTensorRDGs[i];
+	}
+
+	OutputTensorRDGs.Empty();
+	for (int32 i = 0; i < OutputTensorIndices.Num(); ++i)
+	{
+		const int32 Idx = OutputTensorIndices[i];
+		const FTensorDesc& TensorDesc = OutputSymbolicTensors[i];
+		const FTensorShape TensorShape = FTensorShape::MakeFromSymbolic(TensorDesc.GetShape());
+
+		OutputTensorRDGs.Emplace(FTensorRDG::Make(TensorDesc, TensorShape, nullptr));
+		AllTensorRDGs[Idx] = &OutputTensorRDGs[i];
+	}
+
+	checkCode(
+		for (int i = 0; i < AllTensorRDGs.Num(); ++i)
+		{
+			checkf(AllTensorRDGs[i] != nullptr, TEXT("Tensor at index %d, was not allocated for model preparation."), i);
+		};
+	);
+
+	//Allow the specific engine to run shape inference if supported
+	if (PrepareTensorShapesAndData() != 0)
+	{
+		return -1;
+	}
+
+	checkCode(
+		for (int i = 0; i < AllTensorRDGs.Num(); ++i)
+		{
+			checkf(AllTensorRDGs[i] != nullptr, TEXT("Tensor at index %d, was not allocated after model preparation."), i);
+			checkf(AllTensorRDGs[i]->GetShape().IsCompatibleWith(AllSymbolicTensorDescs[i].GetShape()), TEXT("Tensor at index %d have a shape incompatible with model definition."), i);
+		};
+	);
+
+	//Set OutputTensorShapes for the model from preparation result
+	for (int32 OutputIndices : OutputTensorIndices)
+	{
+		OutputTensorShapes.Emplace(AllTensorRDGs[OutputIndices]->GetShape());
+	}
+
+	check(InputTensorIndices.Num() + OutputTensorIndices.Num() + WeightTensorIndices.Num() + IntermediateTensorIndices.Num() == AllTensorRDGs.Num());
+	check(InputTensorShapes.Num() == InputSymbolicTensors.Num());
+	check(OutputTensorShapes.Num() == OutputSymbolicTensors.Num());
+	check(WeightTensorIndices.Num() == WeightTensorRDGs.Num());
+	check(AllTensorRDGs.Num() == AllSymbolicTensorDescs.Num());
+
+	return 0;
+}
+
+FRDGBufferDesc CreateRDGBufferDescForTensorRDG(const FTensorRDG& Tensor)
+{
+	// FIXME: CreateStructuredDesc() creates a crash on VulkanRHI
+	//FRDGBufferDesc Desc = FRDGBufferDesc::CreateStructuredDesc(Tensor.GetElemByteSize(), Tensor.GetVolume());
+	FRDGBufferDesc Desc = FRDGBufferDesc::CreateBufferDesc(Tensor.GetElemByteSize(), Tensor.GetVolume());
+
+	return Desc;
+}
+
+/**
+ * Enqueue operators to RDG, the caller will run the GraphBuilder.Execute()
+ */
+int FModelRDG::EnqueueRDG(FRDGBuilder& RDGBuilder, TConstArrayView<NNECore::FTensorBindingRDG> InInputBindings, TConstArrayView<NNECore::FTensorBindingRDG> InOutputBindings)
+{
+	check(IsInRenderingThread());
+
+	int Res;
+
+	// Verify the model inputs were prepared
+	if (InputTensorShapes.Num() == 0)
+	{
+		UE_LOG(LogNNE, Error, TEXT("EnqueueRDG(): Input shapes are not set, please call SetInputTensorShapes."));
+		return -1;
+	}
+
+	Res = SetTensors(RDGBuilder, InputTensorRDGs, InInputBindings);
+	if (Res != 0)
+	{
+		UE_LOG(LogNNE, Warning, TEXT("Invalid input tensor binding type for tensor index:%d"), Res);
+		return -1;
+	}
+
+	Res = SetTensors(RDGBuilder, OutputTensorRDGs, InOutputBindings);
+	if (Res != 0)
+	{
+		UE_LOG(LogNNE, Warning, TEXT("Invalid output tensor binding type for tensor index:%d"), Res);
+		return -1;
+	}
+
+	//Create buffer for intermediate tensors
+	for (FTensorRDG& TensorRDG : IntermediateTensorRDGs)
+	{
+		const FRDGBufferDesc BufferDesc = CreateRDGBufferDescForTensorRDG(TensorRDG);
+		const FRDGBufferRef TensorBuffer = RDGBuilder.CreateBuffer(BufferDesc, *TensorRDG.GetName(), ERDGBufferFlags::None);
+		check(TensorRDG.GetBuffer() == nullptr);
+		TensorRDG.SetBuffer(TensorBuffer);
+	}
+
+	// TODO: FIXME: DirectML uses RHI buffers instead of RDG buffers
+	//For now weights tensors are not uploaded to GPU thus GetBuffer will return nullptr for them.
+	//checkCode(for (const FTensorRDG* TensorRDG : AllTensorRDGs) { if (TensorRDG != nullptr) { check(TensorRDG->GetBuffer() != nullptr); } });
+
+	//Insert weights tensors
+	for (int32 i = 0; i < WeightTensorIndices.Num(); ++i)
+	{
+		AllTensorRDGs[WeightTensorIndices[i]] = &WeightTensorRDGs[i];
+	}
+
+	// We can now dispatch operators
+	AddDispatchOps_RenderThread(RDGBuilder);
+
+	return 0;
+}
+
+int FModelRDG::SetTensors(FRDGBuilder& GraphBuilder, FTensorRDGArray& InTensorRDGs, TConstArrayView<NNECore::FTensorBindingRDG> InBindings)
+{
+	check(InBindings.Num() == InTensorRDGs.Num());
+
+	for (int32 Idx = 0; Idx < InBindings.Num(); ++Idx)
+	{
+		FTensorRDG& TensorRDG = InTensorRDGs[Idx];
+		const NNECore::FTensorBindingRDG& Binding = InBindings[Idx];
+		TensorRDG.SetBuffer(Binding.Buffer);
+	}
+
+	return 0;
+}
+
+}
\ No newline at end of file
diff --git a/Engine/Plugins/Experimental/NNE/Source/NNERuntimeRDG/Private/NNERuntimeRDGModel.h b/Engine/Plugins/Experimental/NNE/Source/NNERuntimeRDG/Private/NNERuntimeRDGModel.h
new file mode 100644
--- /dev/null
+++ b/Engine/Plugins/Experimental/NNE/Source/NNERuntimeRDG/Private/NNERuntimeRDGModel.h
@@ -0,0 +1,51 @@
+// Copyright Epic Games, Inc. All Rights Reserved.
+
+#pragma once
+
+#include "CoreMinimal.h"
+#include "NNECoreModelBase.h"
+#include "NNECoreRuntimeRDG.h"
+#include "NNECoreTensor.h"
+#include "NNECoreTypes.h"
+#include "NNERuntimeRDG.h"
+
+namespace UE::NNERuntimeRDG::Private
+{
+	class FModelRDG : public NNECore::Internal::FModelBase<NNECore::IModelRDG>
+	{
+	public:
+		FModelRDG() {};
+		virtual ~FModelRDG() = default;
+
+		virtual int SetInputTensorShapes(TConstArrayView<NNECore::FTensorShape> InInputShapes) override;
+		virtual int EnqueueRDG(FRDGBuilder& RDGBuilder, TConstArrayView<NNECore::FTensorBindingRDG> InInputBindings, TConstArrayView<NNECore::FTensorBindingRDG> InOutputBindings) override;
+		
+	protected:
+		bool LoadModel(TConstArrayView<uint8> ModelData, FMLRuntimeFormat& Format, int32 GuidAndVersionSize);
+		int SetTensors(FRDGBuilder& GraphBuilder, FTensorRDGArray& InTensorRDGs, TConstArrayView<NNECore::FTensorBindingRDG> InBindings);
+
+		virtual int PrepareTensorShapesAndData() = 0;
+		virtual void AddDispatchOps_RenderThread(FRDGBuilder& GraphBuilder) = 0;
+
+		//Tensor descriptor
+		TArray<NNECore::FTensorDesc>			AllSymbolicTensorDescs;
+
+		//Tensor indices for models
+		TArray<int32>				IntermediateTensorIndices;
+		TArray<int32>				WeightTensorIndices;
+		TArray<int32>				InputTensorIndices;
+		TArray<int32>				OutputTensorIndices;
+
+		//Tensor indices by operator
+		TArray<TArray<uint32>>		OperatorInputTensorIndices;
+		TArray<TArray<uint32>>		OperatorOutputTensorIndices;
+
+		//RDG Tensors
+		FTensorRDGRefArray			AllTensorRDGs;
+		FTensorRDGArray				InputTensorRDGs;
+		FTensorRDGArray				OutputTensorRDGs;
+		FTensorRDGArray				IntermediateTensorRDGs;
+		FTensorRDGArray				WeightTensorRDGs;
+	};
+	
+} // UE::NNERuntimeRDG::Private
\ No newline at end of file
diff --git a/Engine/Plugins/Experimental/NNE/Source/NNERuntimeRDG/Private/NNERuntimeRDGModelHlsl.cpp b/Engine/Plugins/Experimental/NNE/Source/NNERuntimeRDG/Private/NNERuntimeRDGModelHlsl.cpp
--- a/Engine/Plugins/Experimental/NNE/Source/NNERuntimeRDG/Private/NNERuntimeRDGModelHlsl.cpp
+++ b/Engine/Plugins/Experimental/NNE/Source/NNERuntimeRDG/Private/NNERuntimeRDGModelHlsl.cpp
@@ -2,6 +2,7 @@
 
 #include "NNERuntimeRDGModelHlsl.h"
 #include "NNECoreTensor.h"
+#include "NNERuntimeRDGHlsl.h"
 #include "NNERuntimeRDGHlslOp.h"
 #include "RenderGraphBuilder.h"
 #include "RenderGraphUtils.h"
@@ -9,7 +10,7 @@
 namespace UE::NNERuntimeRDG::Private::Hlsl
 {
 
-namespace
+namespace ModelUtils
 {
 
 FOperatorHlsl* OpCreate(const FString& OpName, TConstArrayView<NNX::FTensorDesc> InputTensorDescs, TConstArrayView<NNX::FTensorDesc> OutputTensorDescs, const NNECore::FAttributeMap& AttributeMap)
@@ -18,7 +19,7 @@ FOperatorHlsl* OpCreate(const FString& OpName, TConstArrayView<NNX::FTensorDesc>
 
 	if (!CreateFn)
 	{
-		UE_LOG(LogNNX, Warning, TEXT("Hlsl MLOperatorRegistry failed to find operator:%s"), *OpName);
+		UE_LOG(LogNNE, Warning, TEXT("Hlsl MLOperatorRegistry failed to find operator:%s"), *OpName);
 		return nullptr;
 	}
 
@@ -26,7 +27,7 @@ FOperatorHlsl* OpCreate(const FString& OpName, TConstArrayView<NNX::FTensorDesc>
 
 	if (!Op->Initialize(InputTensorDescs, OutputTensorDescs, AttributeMap))
 	{
-		UE_LOG(LogNNX, Warning, TEXT("Hlsl engine: Error initializing operator:%s"), *OpName);
+		UE_LOG(LogNNE, Warning, TEXT("Hlsl engine: Error initializing operator:%s"), *OpName);
 		delete Op;
 		return nullptr;
 	}
@@ -62,29 +63,20 @@ void InternAddDispatchOps_RenderThread(FRDGBuilder& GraphBuilder, TConstArrayVie
 	}
 }
 
-int ApplyBinding(TArray<FTensorHLSL>& OutTensorsHLSL, TConstArrayView<NNX::FMLTensorBinding> InBindings, bool bIsInput)
+int ApplyBinding(TArray<FTensorHLSL>& OutTensorsHLSL, TConstArrayView<NNECore::FTensorBindingRDG> InBindings, bool bIsInput)
 {
 	check(OutTensorsHLSL.Num() == InBindings.Num());
 	
 	for (int32 Idx = 0; Idx < InBindings.Num(); ++Idx)
 	{
-		const NNX::FMLTensorBinding& Binding = InBindings[Idx];
+		const NNECore::FTensorBindingRDG& Binding = InBindings[Idx];
 		FTensorHLSL& Tensor = OutTensorsHLSL[Idx];
-
-		if (Binding.BindingType == NNX::EMLTensorBindingDataType::CPUMemory)
-		{
-			if (bIsInput) Tensor.SetUploadBuffer(Binding.CpuMemory);
-			else Tensor.SetDownloadBuffer(Binding.CpuMemory);
-		}
-		else if (Binding.BindingType == NNX::EMLTensorBindingDataType::RDGBuffer)
+		
+		if (Binding.Buffer == nullptr)
 		{
-			Tensor.SetBuffer(Binding.Buffer);
-		}
-		else
-		{
-			// Unsupported tensor binding type
-			return Idx;
+			return -1;
 		}
+		Tensor.SetBuffer(Binding.Buffer);
 	}
 
 	return 0;
@@ -111,8 +103,10 @@ bool FModel::Init(TConstArrayView<uint8> ModelData)
 {
 	check(ModelData.Num() > 0);
 	FMLRuntimeFormat	Format;
-
-	if (!LoadModel(ModelData, Format))
+	int32 GuidSize = sizeof(UNNERuntimeRDGHlslImpl::GUID);
+	int32 VersionSize = sizeof(UNNERuntimeRDGHlslImpl::Version);
+	
+	if (!LoadModel(ModelData, Format, GuidSize+VersionSize))
 	{
 		return false;
 	}
@@ -142,11 +136,11 @@ bool FModel::Init(TConstArrayView<uint8> ModelData)
 			AttributeMap.SetAttribute(Desc.Name, Desc.Value);
 		}
 
-		FOperatorHlsl* Op = OpCreate(TypeName, Inputs, Outputs, AttributeMap);
+		FOperatorHlsl* Op = ModelUtils::OpCreate(TypeName, Inputs, Outputs, AttributeMap);
 
 		if (!Op) //Op.Shader.IsNull())
 		{
-			UE_LOG(LogNNX, Warning, TEXT("Failed to create operator:%s"), *TypeName);
+			UE_LOG(LogNNE, Warning, TEXT("Failed to create operator:%s"), *TypeName);
 
 			// TODO: Cleanup operators
 			return false;
@@ -160,7 +154,7 @@ bool FModel::Init(TConstArrayView<uint8> ModelData)
 
 int FModel::SetInputTensorShapes(TConstArrayView<NNX::FTensorShape> InputShapes)
 {
-	int Res = FInferenceModelRDG::SetInputTensorShapes(InputShapes);
+	int Res = FModelRDG::SetInputTensorShapes(InputShapes);
 	if (Res < 0) return Res;
 
 	auto ConvertTensors = [this] (TArray<FTensorHLSLRef> &TensorRefs, TArray<FTensorHLSL> &Tensors, const FTensorRDGArray &TensorRDGs, const TArray<int32>& TensorIndices)
@@ -190,59 +184,7 @@ int FModel::SetInputTensorShapes(TConstArrayView<NNX::FTensorShape> InputShapes)
 	return 0;
 }
 
-int FModel::RunSync(TConstArrayView<NNX::FMLTensorBinding> InInputBindings, TConstArrayView<NNX::FMLTensorBinding> InOutputBindings)
-{
-	// Verify the model inputs were prepared
-	if (InputTensorShapes.Num() == 0)
-	{
-		UE_LOG(LogNNX, Error, TEXT("Run(): Input shapes are not set, please call SetInputTensorShapes."));
-		return -1;
-	}
-	
-	int Res = 0;
-	FEvent* Signal = FGenericPlatformProcess::GetSynchEventFromPool(false);
-
-	ENQUEUE_RENDER_COMMAND(FMLInferenceModel_Run)
-	(
-		[&Signal, &Res, this, InInputBindings, InOutputBindings](FRHICommandListImmediate& RHICmdList)
-		{
-			TOptional<ERHIPipeline>		Pipeline = RHICmdList.GetPipeline();
-
-			if (Pipeline == ERHIPipeline::None)
-			{
-				RHICmdList.SwitchPipeline(ERHIPipeline::Graphics);
-			}
-
-			FRDGBuilder	RDGBuilder(RHICmdList);
-
-			Res = EnqueueRDG(RDGBuilder, InInputBindings, InOutputBindings);
-			if (Res == 0)
-			{
-				RDGBuilder.Execute();
-
-				// FIXME: Using BlockUntilGPUIdle() prevents hang on Linux
-				// FIXME: Adapt to redesigned readback API (UE 5.2)
-				RHICmdList.BlockUntilGPUIdle();
-
-				for (FTensorHLSLRef Tensor : AllTensorHLSLRefs)
-				{
-					Tensor->Resolve();
-				}
-			}
-			
-			Signal->Trigger();
-		}
-	);
-
-	// We need to wait for render thread to finish
-	Signal->Wait();
-
-	FGenericPlatformProcess::ReturnSynchEventToPool(Signal);
-
-	return Res;
-}
-
-int FModel::EnqueueRDG(FRDGBuilder& GraphBuilder, TConstArrayView<NNX::FMLTensorBinding> InInputBindings, TConstArrayView<NNX::FMLTensorBinding> InOutputBindings)
+int FModel::EnqueueRDG(FRDGBuilder& GraphBuilder, TConstArrayView<NNECore::FTensorBindingRDG> InInputBindings, TConstArrayView<NNECore::FTensorBindingRDG> InOutputBindings)
 {
 	check(IsInRenderingThread());
 
@@ -251,27 +193,27 @@ int FModel::EnqueueRDG(FRDGBuilder& GraphBuilder, TConstArrayView<NNX::FMLTensor
 	// Verify the model inputs were prepared
 	if (InputTensorShapes.Num() == 0)
 	{
-		UE_LOG(LogNNX, Error, TEXT("EnqueueRDG(): Input shapes are not set, please call SetInputTensorShapes."));
+		UE_LOG(LogNNE, Error, TEXT("EnqueueRDG(): Input shapes are not set, please call SetInputTensorShapes."));
 		return -1;
 	}
 
-	Res = ApplyBinding(InputTensorHLSLs, InInputBindings, true);
+	Res = ModelUtils::ApplyBinding(InputTensorHLSLs, InInputBindings, true);
 	if (Res != 0)
 	{
-		UE_LOG(LogNNX, Warning, TEXT("Invalid input tensor binding type for tensor index:%d"), Res);
+		UE_LOG(LogNNE, Warning, TEXT("Invalid input tensor binding for tensor index:%d"), Res);
 		return -1;
 	}
 
-	Res = ApplyBinding(OutputTensorHLSLs, InOutputBindings, false);
+	Res = ModelUtils::ApplyBinding(OutputTensorHLSLs, InOutputBindings, false);
 	if (Res != 0)
 	{
-		UE_LOG(LogNNX, Warning, TEXT("Invalid output tensor binding type for tensor index:%d"), Res);
+		UE_LOG(LogNNE, Warning, TEXT("Invalid output tensor binding for tensor index:%d"), Res);
 		return -1;
 	}
 
-	if (!ApplyWeights(GraphBuilder, WeightTensorHLSLs, WeightsExternalRDGResources))
+	if (!ModelUtils::ApplyWeights(GraphBuilder, WeightTensorHLSLs, WeightsExternalRDGResources))
 	{
-		UE_LOG(LogNNX, Warning, TEXT("Could not register the weights for graph execution."));
+		UE_LOG(LogNNE, Warning, TEXT("Could not register the weights for graph execution."));
 		return -1;
 	}
 
@@ -292,17 +234,7 @@ int FModel::EnqueueRDG(FRDGBuilder& GraphBuilder, TConstArrayView<NNX::FMLTensor
 		}
 	}
 
-	for (FTensorHLSLRef Tensor : AllTensorHLSLRefs)
-	{
-		Tensor->EnqueueUploadRdg(GraphBuilder);
-	}
-
-	InternAddDispatchOps_RenderThread<FOperatorHlsl*>(GraphBuilder, AllTensorHLSLRefs, OperatorInputTensorIndices, OperatorOutputTensorIndices, Operators);
-
-	for (FTensorHLSLRef Tensor : AllTensorHLSLRefs)
-	{
-		Tensor->EnqueueDownloadRdg(GraphBuilder, bUseManualTransitions);
-	}
+	ModelUtils::InternAddDispatchOps_RenderThread<FOperatorHlsl*>(GraphBuilder, AllTensorHLSLRefs, OperatorInputTensorIndices, OperatorOutputTensorIndices, Operators);
 
 	return 0;
 }
@@ -313,7 +245,7 @@ int FModel::PrepareTensorShapesAndData()
 	
 	if (Operators.Num() == 0)
 	{
-		UE_LOG(LogNNX, Warning, TEXT("No operators in model"));
+		UE_LOG(LogNNE, Warning, TEXT("No operators in model"));
 		return -1;
 	}
 
@@ -363,7 +295,7 @@ int FModel::PrepareTensorShapesAndData()
 		{
 			//Operator could not prepare the output tensors, meaning we can't allocate
 			//output buffer before running the model. This engine does not support this.
-			UE_LOG(LogNNX, Warning, TEXT("Could not deduce tensor shapes for this model during shape inference, HLSL engine wont support the model as it need to precompute all shapes for performance reasons."));
+			UE_LOG(LogNNE, Warning, TEXT("Could not deduce tensor shapes for this model during shape inference, HLSL engine wont support the model as it need to precompute all shapes for performance reasons."));
 			AllTensorRDGs.Empty();
 			return -1;
 		}
diff --git a/Engine/Plugins/Experimental/NNE/Source/NNERuntimeRDG/Private/NNERuntimeRDGModelHlsl.h b/Engine/Plugins/Experimental/NNE/Source/NNERuntimeRDG/Private/NNERuntimeRDGModelHlsl.h
--- a/Engine/Plugins/Experimental/NNE/Source/NNERuntimeRDG/Private/NNERuntimeRDGModelHlsl.h
+++ b/Engine/Plugins/Experimental/NNE/Source/NNERuntimeRDG/Private/NNERuntimeRDGModelHlsl.h
@@ -2,7 +2,7 @@
 
 #pragma once
 
-#include "NNERuntimeRDG.h"
+#include "NNERuntimeRDGModel.h"
 #include "NNECoreTypes.h"
 #include "NNERuntimeRDGTensorHlsl.h"
 
@@ -12,7 +12,7 @@ namespace UE::NNERuntimeRDG::Private::Hlsl
 {
 struct FOperatorHlsl;
 
-class FModel : public FInferenceModelRDG
+class FModel : public FModelRDG
 {
 	
 public:
@@ -22,16 +22,7 @@ public:
 	bool Init(TConstArrayView<uint8> ModelData);
 
 	virtual int SetInputTensorShapes(TConstArrayView<NNECore::FTensorShape> InputShapes) override;
-
-	/**
-	 * Run the inference model (synchronous version)
-	 */
-	virtual int RunSync(TConstArrayView<NNX::FMLTensorBinding> InputBindings, TConstArrayView<NNX::FMLTensorBinding> OutputBindings) override;
-	
-	/**
-	 * Enqueue operators to RDG, the caller will run the GraphBuilder.Execute()
-	 */
-	virtual int EnqueueRDG(FRDGBuilder& GraphBuilder, TConstArrayView<NNX::FMLTensorBinding> InputBindings, TConstArrayView<NNX::FMLTensorBinding> OutputBindings) override;
+	virtual int EnqueueRDG(FRDGBuilder& RDGBuilder, TConstArrayView<NNECore::FTensorBindingRDG> InInputBindings, TConstArrayView<NNECore::FTensorBindingRDG> InOutputBindings) override;
 
 protected:
 
diff --git a/Engine/Plugins/Experimental/NNE/Source/NNERuntimeRDG/Private/NNERuntimeRDGModule.cpp b/Engine/Plugins/Experimental/NNE/Source/NNERuntimeRDG/Private/NNERuntimeRDGModule.cpp
--- a/Engine/Plugins/Experimental/NNE/Source/NNERuntimeRDG/Private/NNERuntimeRDGModule.cpp
+++ b/Engine/Plugins/Experimental/NNE/Source/NNERuntimeRDG/Private/NNERuntimeRDGModule.cpp
@@ -1,61 +1,64 @@
 // Copyright Epic Games, Inc. All Rights Reserved.
 
-#include "NNXCore.h"
-#include "NNXRuntime.h"
+#include "NNERuntimeRDGModule.h"
+#include "NNECore.h"
+#include "NNERuntimeRDGDml.h"
+#include "NNERuntimeRDGHlsl.h"
+#include "UObject/WeakInterfacePtr.h"
 
-#include "NNERuntimeRDG.h"
 
-#include "Modules/ModuleManager.h"
-
-//
-//
-//
-class FNNXRuntimeRDGModule : public IModuleInterface
+void FNNERuntimeRDGModule::StartupModule()
 {
-	NNX::IRuntime* DmlRuntime { nullptr };
-	NNX::IRuntime* HlslRuntime{ nullptr };
+	// NNE runtime ORT Cpu startup
+	NNERuntimeRDGHlsl = NewObject<UNNERuntimeRDGHlslImpl>();
+	if (NNERuntimeRDGHlsl.IsValid())
+	{
+		TWeakInterfacePtr<INNERuntime> RuntimeCPUInterface(NNERuntimeRDGHlsl.Get());
 
-public:
+		NNERuntimeRDGHlsl->Init();
+		NNERuntimeRDGHlsl->AddToRoot();
+		UE::NNECore::RegisterRuntime(RuntimeCPUInterface);
+	}
 
-	virtual void StartupModule() override
-	{
 #ifdef NNE_USE_DIRECTML
-		DmlRuntime = UE::NNERuntimeRDG::Private::Dml::FRuntimeDmlStartup();
-
-		if (DmlRuntime)
-		{
-			NNX::RegisterRuntime(DmlRuntime);
-		}
+	// NNE runtime ORT Dml startup
+	NNERuntimeRDGDml = NewObject<UNNERuntimeRDGDmlImpl>();
+	if (NNERuntimeRDGDml.IsValid())
+	{
+		TWeakInterfacePtr<INNERuntime> RuntimeDmlInterface(NNERuntimeRDGDml.Get());
+		
+		bool bRegisterOnlyOperators = UE::NNERuntimeRDG::Private::Dml::FRuntimeDmlStartup();
+		NNERuntimeRDGDml->Init(bRegisterOnlyOperators);
+		NNERuntimeRDGDml->AddToRoot();
+		UE::NNECore::RegisterRuntime(RuntimeDmlInterface);
+	}
 #endif
+}
 
-		HlslRuntime = UE::NNERuntimeRDG::Private::Hlsl::FRuntimeHlslStartup();
+void FNNERuntimeRDGModule::ShutdownModule()
+{
+	// NNE runtime ORT Cpu shutdown
+	if (NNERuntimeRDGHlsl.IsValid())
+	{
+		TWeakInterfacePtr<INNERuntime> RuntimeCPUInterface(NNERuntimeRDGHlsl.Get());
 
-		if (HlslRuntime)
-		{
-			NNX::RegisterRuntime(HlslRuntime);
-		}
+		UE::NNECore::UnregisterRuntime(RuntimeCPUInterface);
+		NNERuntimeRDGHlsl->RemoveFromRoot();
+		NNERuntimeRDGHlsl = TWeakObjectPtr<UNNERuntimeRDGHlslImpl>(nullptr);
 	}
 
-	virtual void ShutdownModule() override
-	{
 #ifdef NNE_USE_DIRECTML
-		if (DmlRuntime)
-		{
-			NNX::UnregisterRuntime(DmlRuntime);
-			DmlRuntime = nullptr;
-		}
+	// NNE runtime ORT Dml shutdown
+	if (NNERuntimeRDGDml.IsValid())
+	{
+		TWeakInterfacePtr<INNERuntime> RuntimeDmlInterface(NNERuntimeRDGDml.Get());
 
-		UE::NNERuntimeRDG::Private::Dml::FRuntimeDmlShutdown();
+		UE::NNECore::UnregisterRuntime(RuntimeDmlInterface);
+		NNERuntimeRDGDml->RemoveFromRoot();
+		NNERuntimeRDGDml = TWeakObjectPtr<UNNERuntimeRDGDmlImpl>(nullptr);
+	}
 #endif
 
-		if (HlslRuntime)
-		{
-			NNX::UnregisterRuntime(HlslRuntime);
-			HlslRuntime = nullptr;
-		}
-
-		UE::NNERuntimeRDG::Private::Hlsl::FRuntimeHlslShutdown();
-	}
-};
+}
 
-IMPLEMENT_MODULE(FNNXRuntimeRDGModule, NNERuntimeRDG)
+IMPLEMENT_MODULE(FNNERuntimeRDGModule, NNERuntimeRDG);
diff --git a/Engine/Plugins/Experimental/NNE/Source/NNERuntimeRDG/Private/NNERuntimeRDGModule.h b/Engine/Plugins/Experimental/NNE/Source/NNERuntimeRDG/Private/NNERuntimeRDGModule.h
new file mode 100644
--- /dev/null
+++ b/Engine/Plugins/Experimental/NNE/Source/NNERuntimeRDG/Private/NNERuntimeRDGModule.h
@@ -0,0 +1,24 @@
+// Copyright Epic Games, Inc. All Rights Reserved.
+
+#pragma once
+
+#include "CoreMinimal.h"
+#include "Modules/ModuleManager.h"
+#include "Templates/SharedPointer.h"
+#include "UObject/WeakObjectPtr.h"
+#include "UObject/WeakObjectPtrTemplates.h"
+
+class UNNERuntimeRDGDmlImpl;
+class UNNERuntimeRDGHlslImpl;
+
+class FNNERuntimeRDGModule : public IModuleInterface
+{
+
+public:
+	TWeakObjectPtr<UNNERuntimeRDGHlslImpl> NNERuntimeRDGHlsl{ nullptr };
+	TWeakObjectPtr<UNNERuntimeRDGDmlImpl> NNERuntimeRDGDml{ nullptr };
+
+	// Begin IModuleInterface
+	virtual void StartupModule() override;
+	virtual void ShutdownModule() override;
+};
\ No newline at end of file
diff --git a/Engine/Plugins/Experimental/NNE/Source/NNERuntimeRDG/Private/NNERuntimeRDGNeuralPostProcessing.cpp b/Engine/Plugins/Experimental/NNE/Source/NNERuntimeRDG/Private/NNERuntimeRDGNeuralPostProcessing.cpp
--- a/Engine/Plugins/Experimental/NNE/Source/NNERuntimeRDG/Private/NNERuntimeRDGNeuralPostProcessing.cpp
+++ b/Engine/Plugins/Experimental/NNE/Source/NNERuntimeRDG/Private/NNERuntimeRDGNeuralPostProcessing.cpp
@@ -1,7 +1,7 @@
 // Copyright Epic Games, Inc. All Rights Reserved.
 
 #include "NNERuntimeRDGNeuralPostProcessing.h"
-#include "NNXCore.h"
+#include "NNECore.h"
 #include "PostProcess/PostProcessMaterial.h"
 #include "PostProcess/SceneRenderTargets.h"
 #include "ScreenPass.h"
@@ -25,36 +25,40 @@ int32 FNNENeuralPostProcessing::Add(FString RuntimeName, UNNEModelData* ModelDat
 {
 	FScopeLock Lock(&CriticalSection);
 
+	UE_LOG(LogNNE, Error, TEXT("DISABLED FOR NOW"));
+	return -1;
+	
+	/*
 	// Create the model
 	NNX::IRuntime* Runtime = NNX::GetRuntime(RuntimeName);
 	if (!Runtime)
 	{
-		UE_LOG(LogNNX, Error, TEXT("FNNENeuralPostProcessing: No runtime '%s' found. Valid runtimes are: "), *RuntimeName);
+		UE_LOG(LogNNE, Error, TEXT("FNNENeuralPostProcessing: No runtime '%s' found. Valid runtimes are: "), *RuntimeName);
 		TArray<NNX::IRuntime*> Runtimes = NNX::GetAllRuntimes();
 		for (int32 i = 0; i < Runtimes.Num(); i++)
 		{
-			UE_LOG(LogNNX, Error, TEXT("- %s"), *Runtimes[i]->GetRuntimeName());
+			UE_LOG(LogNNE, Error, TEXT("- %s"), *Runtimes[i]->GetRuntimeName());
 		}
 		return -1;
 	}
 
 	if (!ModelData)
 	{
-		UE_LOG(LogNNX, Error, TEXT("FNNENeuralPostProcessing: Valid model data required to load the model"));
+		UE_LOG(LogNNE, Error, TEXT("FNNENeuralPostProcessing: Valid model data required to load the model"));
 		return -1;
 	}
 
 	TConstArrayView<uint8> Data = ModelData->GetModelData(RuntimeName);
 	if (Data.Num() < 1)
 	{
-		UE_LOG(LogNNX, Error, TEXT("FNNENeuralPostProcessing: No model data for %s found"), *RuntimeName);
+		UE_LOG(LogNNE, Error, TEXT("FNNENeuralPostProcessing: No model data for %s found"), *RuntimeName);
 		return -1;
 	}
 
 	TSharedPtr<NNX::FMLInferenceModel> Model = TSharedPtr<NNX::FMLInferenceModel>(Runtime->CreateModel(Data).Release());
 	if (!Model.IsValid())
 	{
-		UE_LOG(LogNNX, Error, TEXT("FNNENeuralPostProcessing: Could not create model using %s"), *RuntimeName);
+		UE_LOG(LogNNE, Error, TEXT("FNNENeuralPostProcessing: Could not create model using %s"), *RuntimeName);
 		return -1;
 	}
 
@@ -68,6 +72,7 @@ int32 FNNENeuralPostProcessing::Add(FString RuntimeName, UNNEModelData* ModelDat
 	SetInputSize(LastId, FIntPoint(-1, -1));
 
 	return LastId;
+	*/
 }
 
 bool FNNENeuralPostProcessing::Remove(int32 ModelId)
diff --git a/Engine/Plugins/Experimental/NNE/Source/NNERuntimeRDG/Private/NNERuntimeRDGTensorHlsl.cpp b/Engine/Plugins/Experimental/NNE/Source/NNERuntimeRDG/Private/NNERuntimeRDGTensorHlsl.cpp
--- a/Engine/Plugins/Experimental/NNE/Source/NNERuntimeRDG/Private/NNERuntimeRDGTensorHlsl.cpp
+++ b/Engine/Plugins/Experimental/NNE/Source/NNERuntimeRDG/Private/NNERuntimeRDGTensorHlsl.cpp
@@ -26,7 +26,7 @@ void FTensorHLSL::EnqueueDownloadRdg(FRDGBuilder& GraphBuilder, bool bUseManualT
 
 	Readback = MakeUnique<FRHIGPUBufferReadback>(FName(TEXT("FMLTensorReadback_") + GetName()));
 
-	FMLTensorReadbackParameters* TensorReadbackParams = GraphBuilder.AllocParameters<FMLTensorReadbackParameters>();
+	FNNETensorReadbackParameters* TensorReadbackParams = GraphBuilder.AllocParameters<FNNETensorReadbackParameters>();
 	TensorReadbackParams->Buffer = GetBuffer();
 
 	GraphBuilder.AddPass(
diff --git a/Engine/Plugins/Experimental/NNE/Source/NNXCore/Public/NNECoreRuntimeRDG.h b/Engine/Plugins/Experimental/NNE/Source/NNXCore/Public/NNECoreRuntimeRDG.h
--- a/Engine/Plugins/Experimental/NNE/Source/NNXCore/Public/NNECoreRuntimeRDG.h
+++ b/Engine/Plugins/Experimental/NNE/Source/NNXCore/Public/NNECoreRuntimeRDG.h
@@ -42,7 +42,7 @@ public:
 
 	/** Enqueue the execution on the Render graph render thread. It's caller's responsibility to actually run the graph.
 	* Bindings should point to a buffers big enough */
-	virtual bool EnqueueRDG(FRDGBuilder& RDGBuilder, TConstArrayView<FTensorBindingRDG> Inputs, TConstArrayView<FTensorBindingRDG> Outputs) = 0;
+	virtual int EnqueueRDG(FRDGBuilder& RDGBuilder, TConstArrayView<FTensorBindingRDG> Inputs, TConstArrayView<FTensorBindingRDG> Outputs) = 0;
 };
 
 } // UE::NNECore
