NNE Fix for static analysis 253778

diff --git a/Engine/Plugins/Experimental/NNE/Source/NNEQA/Private/NNEQAModel.h b/Engine/Plugins/Experimental/NNE/Source/NNEQA/Private/NNEQAModel.h
--- a/Engine/Plugins/Experimental/NNE/Source/NNEQA/Private/NNEQAModel.h
+++ b/Engine/Plugins/Experimental/NNE/Source/NNEQA/Private/NNEQAModel.h
@@ -2,7 +2,6 @@
 
 #pragma once
 
-#include "NNEQAModel.h"
 #include "CoreMinimal.h"
 #include "NNECore.h"
 #include "NNECoreRuntimeCPU.h"
