NNE - Added some docs to the public API

diff --git a/Engine/Plugins/Experimental/NNE/Source/NNECore/Public/NNECore.h b/Engine/Plugins/Experimental/NNE/Source/NNECore/Public/NNECore.h
--- a/Engine/Plugins/Experimental/NNE/Source/NNECore/Public/NNECore.h
+++ b/Engine/Plugins/Experimental/NNE/Source/NNECore/Public/NNECore.h
@@ -10,14 +10,39 @@ NNECORE_API DECLARE_LOG_CATEGORY_EXTERN(LogNNE, Log, All);
 
 namespace UE::NNECore
 {
+	/**
+	 * Register a runtime to make it accessible to NNE clients.
+	 *
+	 * The caller needs to keep a strong pointer to the runtime object to prevent it from being garbage collected.
+	 *
+	 * @param Runtime A weak interface pointer to the runtime to be registered.
+	 * @return True if the runtime has been registered successfully, false otherwise.
+	 */
 	NNECORE_API bool RegisterRuntime(TWeakInterfacePtr<INNERuntime> Runtime);
 
+	/**
+	 * Unregister a registered runtime.
+	 *
+	 * @param Runtime A weak interface pointer to the runtime to be unregistered.
+	 * @return True if the runtime has been unregistered successfully, false otherwise (e.g. if the runtime has not been registered).
+	 */
 	NNECORE_API bool UnregisterRuntime(TWeakInterfacePtr<INNERuntime> Runtime);
 	
-	// Enumerate all available runtime modules
+	/**
+	 * List and return all registered runtimes.
+	 *
+	 * @return An array containing weak pointers to all registered runtimes.
+	 */
 	NNECORE_API TArrayView<TWeakInterfacePtr<INNERuntime>> GetAllRuntimes();
 
-	/// Get runtime by name
+	/**
+	 * Find and return a runtime by name and interface.
+	 *
+	 * This function tries to find a runtime by name and casts it to the interface passed as template argument (e.g. INNERuntime, INNERuntimeCPU or INNERuntimeRDG).
+	 *
+	 * @param Name The name of the runtime.
+	 * @return A weak pointer to the runtime if it has been found and implements the interface in the template argument or an invalid pointer otherwise.
+	 */
 	template<class T> TWeakInterfacePtr<T> GetRuntime(const FString& Name)
 	{
 		for (TWeakInterfacePtr<INNERuntime> Runtime : GetAllRuntimes())
diff --git a/Engine/Plugins/Experimental/NNE/Source/NNECore/Public/NNECoreModelData.h b/Engine/Plugins/Experimental/NNE/Source/NNECore/Public/NNECoreModelData.h
--- a/Engine/Plugins/Experimental/NNE/Source/NNECore/Public/NNECoreModelData.h
+++ b/Engine/Plugins/Experimental/NNE/Source/NNECore/Public/NNECoreModelData.h
@@ -6,6 +6,13 @@
 
 #include "NNECoreModelData.generated.h"
 
+/**
+ * This class represents assets that store neural network model data.
+ *
+ * Neural network models typically consist of a graph of operations and corresponding parameters as e.g. weights.
+ * UNNEModelData assets store such model data as imported e.g. by the UNNEModelDataFactory class.
+ * An INNERuntime object retrieved by UE::NNECore::GetRuntime<T>(const FString& Name) can be used to create an inferable neural network model.
+ */
 UCLASS(BlueprintType, Category = "NNE")
 class NNECORE_API UNNEModelData : public UObject
 {
@@ -14,44 +21,51 @@ class NNECORE_API UNNEModelData : public UObject
 public:
 
 	/**
-	 * Function to initialize the data (will do a copy).
-	 * Called by the NNEModelDataFactory.
-	 * Type corresponds to the file extension.
+	 * Initialize the model data with a copy of the data inside Buffer.
+	 *
+	 * This function is called by the UNNEModelDataFactory class when importing a neural netowkr model file.
+	 *
+	 * @param Type A string identifying the type of data inside this asset. Corresponds to the extension of the imported file.
+	 * @param Buffer The raw binary file data of the imported model to be copied into this asset.
 	 */
 	void Init(const FString& Type, TConstArrayView<uint8> Buffer);
 
 	/**
-	 * The function returns the cached (editor) or cooked (game) optimized model data for a given runtime.
-	 * In editor, the function will create the model data with the passed runtime in case it has not been cached yet.
-	 * Returns an empty view in case of a failure.
+	 * Get the cached (editor) or cooked (game) optimized model data for a given runtime.
+	 *
+	 * This function is used by runtimes when creating a model. In editor, the function will create the optimized model data with the passed runtime in case it has not been cached yet. In game, the cooked data is accessed.
+	 *
+	 * @param RuntimeName The name of the runtime for which the data should be returned.
+	 * @return The optimized and runtime specific model data or an empty view in case of a failure.
 	 */
 	TConstArrayView<uint8> GetModelData(const FString& RuntimeName);
 
 public:
 
+	/**
+	 * Implements custom serialization of this asset.
+	 */
 	virtual void Serialize(FArchive& Ar) override;
 
 	/**
-	 * A GUID used for versioning.
+	 * A Guid used for asset versioning.
 	 */
 	const static FGuid GUID;
 
 private:
 
 	/**
-	 * The file type passed by the factory when importing a model.
-	 * Corresponds to the file extension.
+	 * A string identifying the type of data inside this asset. Corresponds to the extension of the imported file.
 	 */
 	FString FileType;
 
 	/**
-	 * Raw binary file data of the imported model.
+	 * The raw binary file data of the imported model.
 	 */
 	TArray<uint8> FileData;
 
 	/**
-	 * Guid that uniquely identifies this model.
-	 * This is used to cache optimized models in the editor.
+	 * A Guid that uniquely identifies this model. This is used to cache optimized models in the editor.
 	 */
 	FGuid FileDataId;
 
diff --git a/Engine/Plugins/Experimental/NNE/Source/NNECore/Public/NNECoreRuntime.h b/Engine/Plugins/Experimental/NNE/Source/NNECore/Public/NNECoreRuntime.h
--- a/Engine/Plugins/Experimental/NNE/Source/NNECore/Public/NNECoreRuntime.h
+++ b/Engine/Plugins/Experimental/NNE/Source/NNECore/Public/NNECoreRuntime.h
@@ -13,14 +13,50 @@ class NNECORE_API UNNERuntime : public UInterface
 	GENERATED_BODY()
 };
 
+/**
+ * The base interface of a neural network runtime. It is returned by UE::NNECore::GetAllRuntimes() giving access to all available runtimes.
+ *
+ * This interface is mainly used internally by UNNEModelData to cook a model given in a file format (e.g. .onnx) into runtime specific model data.
+ * The model data is then stored inside UNNEModelData which then can be used in game to create an inferable model.
+ * See INNERuntimeCPU and INNERuntimeRDG on how to create a model given a UNNEModelData asset.
+ */
 class NNECORE_API INNERuntime
 {
 	GENERATED_BODY()
 
 public:
+
+	/**
+	 * Get the name of the runtime.
+	 *
+	 * @return The name of the runtime.
+	 */
 	virtual FString GetRuntimeName() const = 0;
+	
+	/**
+	 * Check if the runtime is supported on a specific platform.
+	 *
+	 * This function is used internally in the cooking process of UNNEModelData.
+	 *
+	 * @return True if the platform is supported, false otherwise.
+	 */
 	virtual bool IsPlatformSupported(const ITargetPlatform* TargetPlatform) const = 0;
 
+	/**
+	 * Check if the runtime is able to create model data given some file data representing a neural network.
+	 *
+	 * param FileType The type of file inside FileData. Corresponds to the file extension (e.g. 'onnx').
+	 * param FileData The raw binary file of a neural network model.
+	 * @return True if the runtime is able to create model data, false otherwise.
+	 */
 	virtual bool CanCreateModelData(FString FileType, TConstArrayView<uint8> FileData) const = 0;
+
+	/**
+	 * Create model data given some raw file data.
+	 *
+	 * param FileType The type of file inside FileData. Corresponds to the file extension (e.g. 'onnx').
+	 * param FileData The raw binary file of a neural network model.
+	 * @return Data representing the runtime specific representation of the model to be stored by UNNEModelData on success or an empty array otherwise.
+	 */
 	virtual TArray<uint8> CreateModelData(FString FileType, TConstArrayView<uint8> FileData) = 0;
 };
\ No newline at end of file
diff --git a/Engine/Plugins/Experimental/NNE/Source/NNECore/Public/NNECoreRuntimeCPU.h b/Engine/Plugins/Experimental/NNE/Source/NNECore/Public/NNECoreRuntimeCPU.h
--- a/Engine/Plugins/Experimental/NNE/Source/NNECore/Public/NNECoreRuntimeCPU.h
+++ b/Engine/Plugins/Experimental/NNE/Source/NNECore/Public/NNECoreRuntimeCPU.h
@@ -14,35 +14,84 @@ class UNNEModelData;
 namespace UE::NNECore
 {
 
+/**
+ * The tensor binding for passing input and output to IModelCPU.
+ *
+ * Memory is owned by the caller. The caller must make sure the buffer is large enough and matches SizeInBytes.
+ */
 struct NNECORE_API FTensorBindingCPU
 {
 	void*	Data;
 	uint64	SizeInBytes;
 };
 
+/**
+ * The interface of a model that can run on CPU.
+ *
+ * Use UE::NNECore::GetRuntime<INNERuntimeCPU>(RuntimeName) to get a runtime capable of creating CPU models.
+ */
 class NNECORE_API IModelCPU
 {
 public:
 
 	virtual ~IModelCPU() = default;
 
-	/** Getters for tensor description as defined by the model potentially with variable dimensions() */
+	/**
+	 * Get the input tensor descriptions as defined by the model, potentially with variable dimensions.
+	 *
+	 * @return An array containing a tensor descriptor for each input tensor of the model.
+	 */
 	virtual TConstArrayView<FTensorDesc> GetInputTensorDescs() const = 0;
+	
+	/**
+	 * Get the output tensor descriptions as defined by the model, potentially with variable dimensions.
+	 *
+	 * @return An array containing a tensor descriptor for each output tensor of the model.
+	 */
 	virtual TConstArrayView<FTensorDesc> GetOutputTensorDescs() const = 0;
 
-	/** Getters for input shapes if they were set already (see SetInputTensorShapes). Empty list otherwise. */
+	/**
+	 * Get the input shapes.
+	 *
+	 * SetInputTensorShapes must be called prior of running a model.
+	 *
+	 * @return An array of input shapes or an empty array if SetInputTensorShapes has not been called.
+	 */
 	virtual TConstArrayView<FTensorShape> GetInputTensorShapes() const = 0;
 
-	/** Getters for outputs shapes if they were already resolved. Empty list otherwise.
-	* Output shape might be resolved after a call to SetInputTensorShapes() if the model and runtime support it
-	* otherwise they will be resolved during Run() or EnqueueRDG() */
+	/**
+	 * Getters for outputs shapes if they were already resolved.
+	 *
+	 * Output shapes might be resolved after a call to SetInputTensorShapes if the model and runtime supports it.
+	 * Otherwise they will be resolved while running the model
+	 *
+	 * @return An array of output shapes or an empty array if not resolved yet.
+	 */
 	virtual TConstArrayView<FTensorShape> GetOutputTensorShapes() const = 0;
 	
-	/** Prepare the model to be run with the given input shape, The call is mandatory before Run or Enqueue() can be called. */
+	/**
+	 * Prepare the model to be run with the given input shape.
+	 *
+	 * The call is mandatory before a model can be run.
+	 * The function will run shape inference and resolve, if possible, the output shapes which can then be accessed by calling GetOutputTensorShapes().
+	 *
+	 * @param InInputShapes The input shapes to prepare the model with.
+	 * @return 0 on success or a non-zero number otherwise.
+	 */
 	virtual int SetInputTensorShapes(TConstArrayView<FTensorShape> InInputShapes) = 0;
 
-	/** This call is synchronous on all runtime types(CPU, RDG), i.e.the calling thread will be blocked until execution is finished.
-	* Bindings should point to a buffers big enough */
+	/**
+	 * Evaluate the model synchronously.
+	 *
+	 * SetInputTensorShapes must be called prior to this call.
+	 * This function will block the calling thread until the inference is complete.
+	 * The caller owns the memory inside the bindings and must make sure that they are big enough.
+	 * Clients can call this function from an async task but must make sure the memory remains valid throughout the evaluation.
+	 *
+	 * @param InInputTensors An array containing tensor bindings for each input tensor with caller owned memory containing the input data.
+	 * @param InOutputTensors An array containing tensor bindings for each output tensor with caller owned memory big enough to contain the results on success.
+	 * @return 0 on success or a non-zero number otherwise.
+	 */
 	virtual int RunSync(TConstArrayView<FTensorBindingCPU> InInputTensors, TConstArrayView<FTensorBindingCPU> InOutputTensors) = 0;
 };
 
@@ -54,11 +103,34 @@ class NNECORE_API UNNERuntimeCPU : public UInterface
 	GENERATED_BODY()
 };
 
+/**
+ * The interface of a neural network runtime capable of creating CPU models.
+ *
+ * Call UE::NNECore::GetRuntime<INNERuntimeCPU>(RuntimeName) to get a runtime implementing this interface.
+ */
 class NNECORE_API INNERuntimeCPU
 {
 	GENERATED_BODY()
 	
 public:
+
+	/**
+	 * Check if the runtime is able to create a model given some ModelData.
+	 *
+	 * @param ModelData The model data for which to create a model.
+	 * @return True if the runtime is able to create the model, false otherwise.
+	 */
 	virtual bool CanCreateModelCPU(TObjectPtr<UNNEModelData> ModelData) const = 0;
+	
+	/**
+	 * Create a model given some ModelData.
+	 *
+	 * The caller must make sure ModelData remains valid throughout the call.
+	 * ModelData is not required anymore after the model has been created.
+	 * The caller can decide to convert the result into a shared pointer if required (e.g. if the model needs to be shared with an async task for evaluation).
+	 *
+	 * @param ModelData The model data for which to create a model.
+	 * @return A caller owned model representing the neural network stored in ModelData.
+	 */
 	virtual TUniquePtr<UE::NNECore::IModelCPU> CreateModelCPU(TObjectPtr<UNNEModelData> ModelData) = 0;
 };
\ No newline at end of file
diff --git a/Engine/Plugins/Experimental/NNE/Source/NNECore/Public/NNECoreRuntimeRDG.h b/Engine/Plugins/Experimental/NNE/Source/NNECore/Public/NNECoreRuntimeRDG.h
--- a/Engine/Plugins/Experimental/NNE/Source/NNECore/Public/NNECoreRuntimeRDG.h
+++ b/Engine/Plugins/Experimental/NNE/Source/NNECore/Public/NNECoreRuntimeRDG.h
@@ -14,34 +14,83 @@ class UNNEModelData;
 namespace UE::NNECore
 {
 
+/**
+ * The tensor binding for passing input and output to IModelRDG.
+ *
+ * Memory is owned by the caller. The caller must make sure the buffer is large enough.
+ */
 struct NNECORE_API FTensorBindingRDG
 {
 	FRDGBufferRef Buffer;
 };
 
+/**
+ * The interface of a model that can run on RDG.
+ *
+ * Use UE::NNECore::GetRuntime<INNERuntimeRDG>(RuntimeName) to get a runtime capable of creating RDG models.
+ */
 class NNECORE_API IModelRDG
 {
 public:
 
 	virtual ~IModelRDG() = default;
 
-	/** Getters for tensor description as defined by the model potentially with variable dimensions() */
+	/**
+	 * Get the input tensor descriptions as defined by the model, potentially with variable dimensions.
+	 *
+	 * @return An array containing a tensor descriptor for each input tensor of the model.
+	 */
 	virtual TConstArrayView<FTensorDesc> GetInputTensorDescs() const = 0;
+
+	/**
+	 * Get the output tensor descriptions as defined by the model, potentially with variable dimensions.
+	 *
+	 * @return An array containing a tensor descriptor for each output tensor of the model.
+	 */
 	virtual TConstArrayView<FTensorDesc> GetOutputTensorDescs() const = 0;
 
-	/** Getters for input shapes if they were set already (see SetInputTensorShapes). Empty list otherwise. */
+	/**
+	 * Get the input shapes.
+	 *
+	 * SetInputTensorShapes must be called prior of running a model.
+	 *
+	 * @return An array of input shapes or an empty array if SetInputTensorShapes has not been called.
+	 */
 	virtual TConstArrayView<FTensorShape> GetInputTensorShapes() const = 0;
 
-	/** Getters for outputs shapes if they were already resolved. Empty list otherwise.
-	* Output shape might be resolved after a call to SetInputTensorShapes() if the model and runtime support it
-	* otherwise they will be resolved during Run() or EnqueueRDG() */
+	/**
+	 * Getters for outputs shapes if they were already resolved.
+	 *
+	 * Output shapes might be resolved after a call to SetInputTensorShapes if the model and runtime supports it.
+	 * Otherwise they will be resolved while running the model
+	 *
+	 * @return An array of output shapes or an empty array if not resolved yet.
+	 */
 	virtual TConstArrayView<FTensorShape> GetOutputTensorShapes() const = 0;
 	
-	/** Prepare the model to be run with the given input shape, The call is mandatory before Run or Enqueue() can be called. */
+	/**
+	 * Prepare the model to be run with the given input shape.
+	 *
+	 * The call is mandatory before a model can be run.
+	 * The function will run shape inference and resolve, if possible, the output shapes which can then be accessed by calling GetOutputTensorShapes().
+	 *
+	 * @param InInputShapes The input shapes to prepare the model with.
+	 * @return 0 on success or a non-zero number otherwise.
+	 */
 	virtual int SetInputTensorShapes(TConstArrayView<FTensorShape> InInputShapes) = 0;
 
-	/** Enqueue the execution on the Render graph render thread. It's caller's responsibility to actually run the graph.
-	* Bindings should point to a buffers big enough */
+	/**
+	 * Enqueue the model graph to a FRDGBuilder.
+	 *
+	 * This function must be called from the render thread.
+	 * SetInputTensorShapes must be called prior to this call.
+	 * The caller owns the memory inside the bindings and must make sure that they are big enough.
+	 *
+	 * @param RDGBuilder The RDG builder to which the neural network operations are enqueued.
+	 * @param InInputTensors An array containing tensor bindings for each input tensor with caller owned memory containing the input data.
+	 * @param InOutputTensors An array containing tensor bindings for each output tensor with caller owned memory big enough to contain the results on success.
+	 * @return 0 on success or a non-zero number otherwise.
+	 */
 	virtual int EnqueueRDG(FRDGBuilder& RDGBuilder, TConstArrayView<FTensorBindingRDG> Inputs, TConstArrayView<FTensorBindingRDG> Outputs) = 0;
 };
 
@@ -53,11 +102,34 @@ class NNECORE_API UNNERuntimeRDG : public UInterface
 	GENERATED_BODY()
 };
 
+/**
+ * The interface of a neural network runtime capable of creating RDG models.
+ *
+ * Call UE::NNECore::GetRuntime<INNERuntimeRDG>(RuntimeName) to get a runtime implementing this interface.
+ */
 class NNECORE_API INNERuntimeRDG
 {
 	GENERATED_BODY()
 
 public:
+
+	/**
+	 * Check if the runtime is able to create a model given some ModelData.
+	 *
+	 * @param ModelData The model data for which to create a model.
+	 * @return True if the runtime is able to create the model, false otherwise.
+	 */
 	virtual bool CanCreateModelRDG(TObjectPtr<UNNEModelData> ModelData) const = 0;
+	
+	/**
+	 * Create a model given some ModelData.
+	 *
+	 * The caller must make sure ModelData remains valid throughout the call.
+	 * ModelData is not required anymore after the model has been created.
+	 * The caller can decide to convert the result into a shared pointer if required (e.g. if the model needs to be shared with the render thread).
+	 *
+	 * @param ModelData The model data for which to create a model.
+	 * @return A caller owned model representing the neural network stored in ModelData.
+	 */
 	virtual TUniquePtr<UE::NNECore::IModelRDG> CreateModelRDG(TObjectPtr<UNNEModelData> ModelData) = 0;
 };
\ No newline at end of file
diff --git a/Engine/Plugins/Experimental/NNE/Source/NNECore/Public/NNECoreTypes.h b/Engine/Plugins/Experimental/NNE/Source/NNECore/Public/NNECoreTypes.h
--- a/Engine/Plugins/Experimental/NNE/Source/NNECore/Public/NNECoreTypes.h
+++ b/Engine/Plugins/Experimental/NNE/Source/NNECore/Public/NNECoreTypes.h
@@ -5,6 +5,9 @@
 #include "CoreMinimal.h"
 #include "UObject/Class.h"
 
+/**
+ * The enum lists all different data types used in NNE 
+ */
 UENUM()
 enum class ENNETensorDataType : uint8
 {
@@ -29,6 +32,12 @@ enum class ENNETensorDataType : uint8
 
 namespace UE::NNECore
 {
+
+	/**
+	 * A symbolic tensor shape represents the shape of a tensor with potentially variable dimension.
+	 * 
+	 * The variable dimensions are represented by negative values. 
+	 */
 	class NNECORE_API FSymbolicTensorShape
 	{
 	public:
@@ -49,6 +58,11 @@ namespace UE::NNECore
 		void operator=(const FSymbolicTensorShape& OtherShape);
 	};
 
+	/**
+	 * The concrete shape of a tensor.
+	 *
+	 * Concrete tensor shapes are well defined through strictly positive values and thus have a well defined volume.
+	 */
 	class NNECORE_API FTensorShape
 	{
 	public:
@@ -71,9 +85,17 @@ namespace UE::NNECore
 		void operator=(const FTensorShape& OtherShape);
 	};
 
-	/** Return data size in bytes for tensor data type */
+	/**
+	 * Return the data size in bytes for a tensor data type.
+	 */
 	size_t NNECORE_API GetTensorDataTypeSizeInBytes(ENNETensorDataType InType);
 
+	/**
+	 * The descriptor for a tensor.
+	 *
+	 * A tensor is described by its name, the type of data it contains and it's shape.
+	 * Since input and output tensors of a neural network can have dynamic shapes, Shape is symbolic.
+	 */
 	class NNECORE_API FTensorDesc
 	{
 		FString					Name;
diff --git a/Engine/Plugins/Experimental/NNE/Source/NNEEditorTools/Public/NNEModelDataFactory.h b/Engine/Plugins/Experimental/NNE/Source/NNEEditorTools/Public/NNEModelDataFactory.h
--- a/Engine/Plugins/Experimental/NNE/Source/NNEEditorTools/Public/NNEModelDataFactory.h
+++ b/Engine/Plugins/Experimental/NNE/Source/NNEEditorTools/Public/NNEModelDataFactory.h
@@ -7,6 +7,11 @@
 
 #include "NNEModelDataFactory.generated.h"
 
+/**
+ * Class for importing and creating new UNNEModelData assets.
+ *
+ * The currently supported format is .onnx. The factory is invoked when a file of this format is dragged to the editor's content browser.
+ */
 UCLASS()
 class NNEEDITORTOOLS_API UNNEModelDataFactory : public UFactory
 {
