NNE - Public API Doc

Add more doc to NNE public API

diff --git a/Engine/Plugins/Experimental/NNE/Source/NNECore/Public/NNECoreModelData.h b/Engine/Plugins/Experimental/NNE/Source/NNECore/Public/NNECoreModelData.h
--- a/Engine/Plugins/Experimental/NNE/Source/NNECore/Public/NNECoreModelData.h
+++ b/Engine/Plugins/Experimental/NNE/Source/NNECore/Public/NNECoreModelData.h
@@ -23,7 +23,7 @@ public:
 	/**
 	 * Initialize the model data with a copy of the data inside Buffer.
 	 *
-	 * This function is called by the UNNEModelDataFactory class when importing a neural netowkr model file.
+	 * This function is called by the UNNEModelDataFactory class when importing a neural network model file.
 	 *
 	 * @param Type A string identifying the type of data inside this asset. Corresponds to the extension of the imported file.
 	 * @param Buffer The raw binary file data of the imported model to be copied into this asset.
@@ -33,7 +33,7 @@ public:
 	/**
 	 * Get the cached (editor) or cooked (game) optimized model data for a given runtime.
 	 *
-	 * This function is used by runtimes when creating a model. In editor, the function will create the optimized model data with the passed runtime in case it has not been cached yet. In game, the cooked data is accessed.
+	 * This function is used by runtimes when creating a model. In editor, the function will create the optimized model data with the passed runtime in case it has not been cached in the DCC yet. In game, the cooked data is accessed.
 	 *
 	 * @param RuntimeName The name of the runtime for which the data should be returned.
 	 * @return The optimized and runtime specific model data or an empty view in case of a failure.
@@ -44,6 +44,7 @@ public:
 
 	/**
 	 * Implements custom serialization of this asset.
+	 * @param Ar The archive to serialize from/to.
 	 */
 	virtual void Serialize(FArchive& Ar) override;
 
diff --git a/Engine/Plugins/Experimental/NNE/Source/NNECore/Public/NNECoreRuntime.h b/Engine/Plugins/Experimental/NNE/Source/NNECore/Public/NNECoreRuntime.h
--- a/Engine/Plugins/Experimental/NNE/Source/NNECore/Public/NNECoreRuntime.h
+++ b/Engine/Plugins/Experimental/NNE/Source/NNECore/Public/NNECoreRuntime.h
@@ -38,6 +38,7 @@ public:
 	 *
 	 * This function is used internally in the cooking process of UNNEModelData.
 	 *
+	 * @param TargetPlatform Interface identifying the target platforms.
 	 * @return True if the platform is supported, false otherwise.
 	 */
 	virtual bool IsPlatformSupported(const ITargetPlatform* TargetPlatform) const = 0;
@@ -45,8 +46,8 @@ public:
 	/**
 	 * Check if the runtime is able to create model data given some file data representing a neural network.
 	 *
-	 * param FileType The type of file inside FileData. Corresponds to the file extension (e.g. 'onnx').
-	 * param FileData The raw binary file of a neural network model.
+	 * @param FileType The type of file inside FileData. Corresponds to the file extension (e.g. 'onnx').
+	 * @param FileData The raw binary file of a neural network model.
 	 * @return True if the runtime is able to create model data, false otherwise.
 	 */
 	virtual bool CanCreateModelData(FString FileType, TConstArrayView<uint8> FileData) const = 0;
@@ -54,8 +55,8 @@ public:
 	/**
 	 * Create model data given some raw file data.
 	 *
-	 * param FileType The type of file inside FileData. Corresponds to the file extension (e.g. 'onnx').
-	 * param FileData The raw binary file of a neural network model.
+	 * @param FileType The type of file inside FileData. Corresponds to the file extension (e.g. 'onnx').
+	 * @param FileData The raw binary file of a neural network model.
 	 * @return Data representing the runtime specific representation of the model to be stored by UNNEModelData on success or an empty array otherwise.
 	 */
 	virtual TArray<uint8> CreateModelData(FString FileType, TConstArrayView<uint8> FileData) = 0;
diff --git a/Engine/Plugins/Experimental/NNE/Source/NNECore/Public/NNECoreRuntimeCPU.h b/Engine/Plugins/Experimental/NNE/Source/NNECore/Public/NNECoreRuntimeCPU.h
--- a/Engine/Plugins/Experimental/NNE/Source/NNECore/Public/NNECoreRuntimeCPU.h
+++ b/Engine/Plugins/Experimental/NNE/Source/NNECore/Public/NNECoreRuntimeCPU.h
@@ -17,7 +17,7 @@ namespace UE::NNECore
 /**
  * The tensor binding for passing input and output to IModelCPU.
  *
- * Memory is owned by the caller. The caller must make sure the buffer is large enough and matches SizeInBytes.
+ * Memory is owned by the caller. The caller must make sure the buffer is large enough and at least as large as SizeInBytes.
  */
 struct NNECORE_API FTensorBindingCPU
 {
@@ -130,7 +130,7 @@ public:
 	 * The caller can decide to convert the result into a shared pointer if required (e.g. if the model needs to be shared with an async task for evaluation).
 	 *
 	 * @param ModelData The model data for which to create a model.
-	 * @return A caller owned model representing the neural network stored in ModelData.
+	 * @return A caller owned model representing the neural network created from ModelData.
 	 */
 	virtual TUniquePtr<UE::NNECore::IModelCPU> CreateModelCPU(TObjectPtr<UNNEModelData> ModelData) = 0;
 };
\ No newline at end of file
diff --git a/Engine/Plugins/Experimental/NNE/Source/NNECore/Public/NNECoreRuntimeRDG.h b/Engine/Plugins/Experimental/NNE/Source/NNECore/Public/NNECoreRuntimeRDG.h
--- a/Engine/Plugins/Experimental/NNE/Source/NNECore/Public/NNECoreRuntimeRDG.h
+++ b/Engine/Plugins/Experimental/NNE/Source/NNECore/Public/NNECoreRuntimeRDG.h
@@ -129,7 +129,7 @@ public:
 	 * The caller can decide to convert the result into a shared pointer if required (e.g. if the model needs to be shared with the render thread).
 	 *
 	 * @param ModelData The model data for which to create a model.
-	 * @return A caller owned model representing the neural network stored in ModelData.
+	 * @return A caller owned model representing the neural network created from ModelData.
 	 */
 	virtual TUniquePtr<UE::NNECore::IModelRDG> CreateModelRDG(TObjectPtr<UNNEModelData> ModelData) = 0;
 };
\ No newline at end of file
diff --git a/Engine/Plugins/Experimental/NNE/Source/NNECore/Public/NNECoreTypes.h b/Engine/Plugins/Experimental/NNE/Source/NNECore/Public/NNECoreTypes.h
--- a/Engine/Plugins/Experimental/NNE/Source/NNECore/Public/NNECoreTypes.h
+++ b/Engine/Plugins/Experimental/NNE/Source/NNECore/Public/NNECoreTypes.h
@@ -6,7 +6,7 @@
 #include "UObject/Class.h"
 
 /**
- * The enum lists all different data types used in NNE 
+ * The enum lists all tensor data types used in NNE 
  */
 UENUM()
 enum class ENNETensorDataType : uint8
@@ -36,62 +36,171 @@ namespace UE::NNECore
 	/**
 	 * A symbolic tensor shape represents the shape of a tensor with potentially variable dimension.
 	 * 
-	 * The variable dimensions are represented by negative values. 
+	 * The variable dimensions are represented by -1 values. 
 	 */
 	class NNECORE_API FSymbolicTensorShape
 	{
 	public:
+		/**
+		 * The maximum number of dimensions supported by the symbolic tensor shape, aka the rank.
+		 */
 		constexpr static int32	MaxRank = 8;
 
 	private:
 		TArray<int32, TInlineAllocator<MaxRank>> Data;
 
 	public:
+		/**
+		 * Construct a symbolic tensor shape with the given dimensions.
+		 *
+		 * @param Data an array of dimensions.
+		 * @return a symbolic tensor shape with the given dimensions.
+		 */
 		static FSymbolicTensorShape Make(TConstArrayView<int32> Data);
 
+		/**
+		 * Get the dimensions of the symbolic tensor shape.
+		 *
+		 * @return a view on the dimensions of the symbolic tensor shape.
+		 */
 		inline TConstArrayView<int32> GetData() const { return Data; };
+		
+		/**
+		 * Get the number of dimensions of the symbolic tensor shape, aka the rank.
+		 *
+		 * @return the number of dimensions of the symbolic tensor shape.
+		 */
 		inline int32 Rank() const { return Data.Num(); }
 
+		/**
+		 * Check if the symbolic tensor shape has no variable dimensions, aka is concrete.
+		 *
+		 * @return true if the symbolic tensor shape is concrete, false otherwise.
+		 */
 		bool IsConcrete() const;
+
+		/**
+		 * Check if the symbolic tensor shape is equal to another one.
+		 *
+		 * @param OtherShape the other symbolic tensor shape.
+		 * @return true if the symbolic tensor shape is equal to the other one, false otherwise.
+		 */
 		bool operator==(const FSymbolicTensorShape& OtherShape) const;
+		
+		/**
+		 * Check if the symbolic tensor shape is different from another one.
+		 * 
+		 * @param OtherShape the other symbolic tensor shape.
+		 * @return true if the symbolic tensor shape is different from the other one, false otherwise.
+		 */
 		bool operator!=(const FSymbolicTensorShape& OtherShape) const;
+
+		/**
+		 * Assign this symbolic tensor shape from another one.
+		 * 
+		 * @param OtherShape the other symbolic tensor shape to copy from.
+		 */
 		void operator=(const FSymbolicTensorShape& OtherShape);
 	};
 
 	/**
 	 * The concrete shape of a tensor.
 	 *
-	 * Concrete tensor shapes are well defined through strictly positive values and thus have a well defined volume.
+	 * Concrete tensor shapes are well defined through strictly positive values and thus have a defined volume.
 	 */
 	class NNECORE_API FTensorShape
 	{
 	public:
+		/**
+		 * The maximum number of dimensions supported by the tensor shape, aka the rank.
+		 */
 		constexpr static int32	MaxRank = FSymbolicTensorShape::MaxRank;
 
 	private:
 		TArray<uint32, TInlineAllocator<MaxRank>> Data;
 
 	public:
+		/**
+		 * Construct a tensor shape with the given dimensions.
+		 *
+		 * @param Data an array of dimensions.
+		 * @return a tensor shape with the given dimensions.
+		 */
 		static FTensorShape Make(TConstArrayView<uint32> Data);
+		
+		/**
+		 * Construct a tensor shape from a symbolic tensor shape.
+		 * 
+		 * Negative dimension values are replaced by 1.
+		 *
+		 * @param SymbolicShape a symbolic tensor shape.
+		 * @return a tensor shape with the given dimensions.
+		 */
 		static FTensorShape MakeFromSymbolic(const FSymbolicTensorShape& SymbolicShape);
 
+		/**
+		 * Get the dimensions of the tensor shape.
+		 *
+		 * @return a view on the dimensions of the tensor shape.
+		 */
 		inline TConstArrayView<uint32> GetData() const { return Data; };
+		
+		/**
+		 * Get the number of dimensions of the tensor shape, aka the rank.
+		 *
+		 * @return the number of dimensions of the tensor shape.
+		 */
 		inline int32 Rank() const { return Data.Num(); }
 		
+		/**
+		 * Get the number of elements in this tensor shape.
+		 *
+		 * @return the number of element in this tensor shape.
+		 */
 		uint64 Volume() const;
+		
+		/**
+		 * Check if this tensor shape is part of the ensemble of shapes defined by the symbolic shape.
+		 * 
+		 * For example [1,2] is compatible with [-1,2] and [1,2] but not with [2,2] or [1,2,3].
+		 *
+		 * @return true if this tensor shape is compatible with the symbolic shape, false otherwise.
+		 */
 		bool IsCompatibleWith(const FSymbolicTensorShape& SymbolicShape) const;
+
+		/**
+		 * Check if the tensor shape is equal to another one.
+		 * 
+		 * @param OtherShape the other tensor shape.
+		 * @return true if the tensor shape is equal to the other one, false otherwise.
+		 */
 		bool operator==(const FTensorShape& OtherShape) const;
+		
+		/**
+		 * Check if the tensor shape is different from another one.
+		 * 
+		 * @param OtherShape the other tensor shape.
+		 * @return true if the tensor shape is different from the other one, false otherwise.
+		 */
 		bool operator!=(const FTensorShape& OtherShape) const;
+
+		/**
+		 * Assign this tensor shape from another one.
+		 * @param OtherShape the other tensor shape to copy from.
+		 */
 		void operator=(const FTensorShape& OtherShape);
 	};
 
 	/**
 	 * Return the data size in bytes for a tensor data type.
+	 * 
+	 * @param InType the type of the element to consider.
+	 * @return the data size in bytes of an element.
 	 */
 	size_t NNECORE_API GetTensorDataTypeSizeInBytes(ENNETensorDataType InType);
 
 	/**
-	 * The descriptor for a tensor.
+	 * The descriptor for a tensor as model inputs and output.
 	 *
 	 * A tensor is described by its name, the type of data it contains and it's shape.
 	 * Since input and output tensors of a neural network can have dynamic shapes, Shape is symbolic.
@@ -105,12 +214,42 @@ namespace UE::NNECore
 		FTensorDesc() = default;
 
 	public:
+		/**
+		 * Construct a tensor description for model input/output
+		 *
+		 * @param Name the name of the tensor.
+		 * @param Shape the symbolic shape of the tensor.
+		 * @param DataType the type of data the tensor contains.
+		 * @return a tensor shape with the given dimensions.
+		 */
 		static FTensorDesc Make(const FString& Name, const FSymbolicTensorShape& Shape, ENNETensorDataType DataType);
-		
+		/**
+		 * Get the name of the tensor
+		 * 
+		 * @return the name of the tensor.
+		 */
 		inline const FString& GetName() const { return Name; }
+		
+		/**
+		 * Get the data type of the tensor
+		 * 
+		 * @return the data type of the tensor.
+		 */
 		inline ENNETensorDataType GetDataType() const { return DataType; }
+		
+		/**
+		 * Get the size in bytes of one element of the tensor.
+		 * 
+		 * @return the size in bytes of one element of the tensor.
+		 */
 		inline uint32 GetElemByteSize() const { return GetTensorDataTypeSizeInBytes(DataType); }
+
+		/**
+		 * Get the symbolic shape of the tensor.
+		 * 
+		 * @return the symbolic shape of the tensor.
+		 */
 		inline const FSymbolicTensorShape& GetShape() const { return Shape; }
 	};
 
-} // namespace NNE
+} // namespace UE::NNECore
