NNE: Allow on the fly model data creation for QA.

Problem: model can't be created on the fly on device, preventing QA tests to run
Solution: allow to create model on the fly from the NNE runtimes.
Test: all nne test green on desktop standalone

#rn ue-nni

diff --git a/Engine/Plugins/Experimental/NNE/Source/NNECore/Private/NNECoreModelData.cpp b/Engine/Plugins/Experimental/NNE/Source/NNECore/Private/NNECoreModelData.cpp
--- a/Engine/Plugins/Experimental/NNE/Source/NNECore/Private/NNECoreModelData.cpp
+++ b/Engine/Plugins/Experimental/NNE/Source/NNECore/Private/NNECoreModelData.cpp
@@ -87,15 +87,14 @@ void UNNEModelData::Init(const FString& Type, TConstArrayView<uint8> Buffer)
 
 TConstArrayView<uint8> UNNEModelData::GetModelData(const FString& RuntimeName)
 {
-#if WITH_EDITOR
-
 	// Check if we have a local cache hit
 	TArray<uint8>* LocalData = ModelData.Find(RuntimeName);
 	if (LocalData)
 	{
 		return TConstArrayView<uint8>(LocalData->GetData(), LocalData->Num());
 	}
-
+	
+#if WITH_EDITOR
 	// Check if we have a remote cache hit
 	FSharedBuffer RemoteData = GetFromDDC(FileDataId, RuntimeName);
 	if (RemoteData.GetSize() > 0)
@@ -105,7 +104,8 @@ TConstArrayView<uint8> UNNEModelData::GetModelData(const FString& RuntimeName)
 		TArray<uint8>* CachedRemoteData = ModelData.Find(RuntimeName);
 		return TConstArrayView<uint8>(CachedRemoteData->GetData(), CachedRemoteData->Num());
 	}
-
+#endif
+	
 	// Try to create the model
 	TArray<uint8> CreatedData = Create(RuntimeName, FileType, FileData);
 	if (CreatedData.Num() < 1)
@@ -113,25 +113,17 @@ TConstArrayView<uint8> UNNEModelData::GetModelData(const FString& RuntimeName)
 		return {};
 	}
 
-	// Cache the model and put it into ddc
+	// Cache the model
 	ModelData.Add(RuntimeName, CreatedData);
 
+#if WITH_EDITOR
+	// And put it into DDC
 	FSharedBuffer SharedBuffer = MakeSharedBufferFromArray(MoveTemp(CreatedData));
 	PutIntoDDC(FileDataId, RuntimeName, SharedBuffer);
-
+#endif
+	
 	TArray<uint8>* CachedCreatedData = ModelData.Find(RuntimeName);
 	return TConstArrayView<uint8>(CachedCreatedData->GetData(), CachedCreatedData->Num());
-
-#else
-
-	TArray<uint8>* Data = ModelData.Find(RuntimeName);
-	if (Data)
-	{
-		return TConstArrayView<uint8>(Data->GetData(), Data->Num());
-	}
-	return {};
-
-#endif
 }
 
 void UNNEModelData::Serialize(FArchive& Ar)
