NNE - Rename rdg runtimes

diff --git a/Engine/Plugins/Experimental/NNE/Source/NNERuntimeRDG/Private/NNERuntimeRDGDml.h b/Engine/Plugins/Experimental/NNE/Source/NNERuntimeRDG/Private/NNERuntimeRDGDml.h
--- a/Engine/Plugins/Experimental/NNE/Source/NNERuntimeRDG/Private/NNERuntimeRDGDml.h
+++ b/Engine/Plugins/Experimental/NNE/Source/NNERuntimeRDG/Private/NNERuntimeRDGDml.h
@@ -33,7 +33,7 @@ public:
 
 	bool Init(bool bRegisterOnlyOperators);
 
-	virtual FString GetRuntimeName() const override { return TEXT("NNERuntimeDml"); };
+	virtual FString GetRuntimeName() const override { return TEXT("NNERuntimeRDGDml"); };
 	virtual bool IsPlatformSupported(const ITargetPlatform* TargetPlatform) const override { return true; };
 
 	virtual bool CanCreateModelData(FString FileType, TConstArrayView<uint8> FileData) const override;
diff --git a/Engine/Plugins/Experimental/NNE/Source/NNERuntimeRDG/Private/NNERuntimeRDGHlsl.h b/Engine/Plugins/Experimental/NNE/Source/NNERuntimeRDG/Private/NNERuntimeRDGHlsl.h
--- a/Engine/Plugins/Experimental/NNE/Source/NNERuntimeRDG/Private/NNERuntimeRDGHlsl.h
+++ b/Engine/Plugins/Experimental/NNE/Source/NNERuntimeRDG/Private/NNERuntimeRDGHlsl.h
@@ -26,7 +26,7 @@ public:
 
 	bool Init();
 
-	virtual FString GetRuntimeName() const override { return TEXT("NNERuntimeHlsl"); };
+	virtual FString GetRuntimeName() const override { return TEXT("NNERuntimeRDGHlsl"); };
 	virtual bool IsPlatformSupported(const ITargetPlatform* TargetPlatform) const override { return true; };
 
 	virtual bool CanCreateModelData(FString FileType, TConstArrayView<uint8> FileData) const override;
