NNE - Update the runtime names in our test bench

diff --git a/Engine/Plugins/Experimental/NNE/Source/NNEQA/Private/NNEQAUtils.cpp b/Engine/Plugins/Experimental/NNE/Source/NNEQA/Private/NNEQAUtils.cpp
--- a/Engine/Plugins/Experimental/NNE/Source/NNEQA/Private/NNEQAUtils.cpp
+++ b/Engine/Plugins/Experimental/NNE/Source/NNEQA/Private/NNEQAUtils.cpp
@@ -558,7 +558,7 @@ namespace UE::NNEQA::Private
 			else
 			{
 				bool bShouldRunVariadicTest = (ONNXModelVariadic.Format != ENNEInferenceFormat::Invalid);
-				bShouldRunVariadicTest &= !(RuntimeName == "NNERuntimeDml");
+				bShouldRunVariadicTest &= !(RuntimeName == "NNERuntimeRDGDml");
 				
 				bool bTestSuceeded = RunTestInferenceAndCompareToRef(TestSetup, Runtime->GetRuntimeName(), ONNXModel, RefOutputMemBuffers, RefOutputTensors);
 
diff --git a/Engine/Plugins/Experimental/NNE/Source/NNEQA/Resources/NNEQATestDesc.json b/Engine/Plugins/Experimental/NNE/Source/NNEQA/Resources/NNEQATestDesc.json
--- a/Engine/Plugins/Experimental/NNE/Source/NNEQA/Resources/NNEQATestDesc.json
+++ b/Engine/Plugins/Experimental/NNE/Source/NNEQA/Resources/NNEQATestDesc.json
@@ -10,8 +10,8 @@
 					"tags": [ "SmokeTest" ],
 					"runtimes": [
 						{ "name": "NNERuntimeORTDml" },
-						{ "name": "NNERuntimeHlsl",  "skip": true, "comment": "Missing BatchNormalisation operator." },
-						{ "name": "NNERuntimeDml",  "skip": true, "comment": "Model not supported yet." }
+						{ "name": "NNERuntimeRDGHlsl",  "skip": true, "comment": "Missing BatchNormalisation operator." },
+						{ "name": "NNERuntimeRDGDml",  "skip": true, "comment": "Model not supported yet." }
 					],
 					"absolute_tolerance": 1e-5
 				},
@@ -19,9 +19,9 @@
 					"target": "Mosaic-9",
 					"runtimes": [
 						{ "name": "NNERuntimeORTDml" },
-						{ "name": "NNERuntimeORTCpuEditor" },
-						{ "name": "NNERuntimeHlsl" },
-						{ "name": "NNERuntimeDml",  "skip": true, "comment": "Model not supported yet." }
+						{ "name": "NNERuntimeORTCpu" },
+						{ "name": "NNERuntimeRDGHlsl" },
+						{ "name": "NNERuntimeRDGDml",  "skip": true, "comment": "Model not supported yet." }
 					],
 					"absolute_tolerance": 5.0e-2
 				}
@@ -36,8 +36,8 @@
 					"tags": [ "SmokeTest" ],
 					"runtimes": [
 						{ "name": "NNERuntimeORTDml" },
-						{ "name": "NNERuntimeHlsl" },
-						{ "name": "NNERuntimeDml" }
+						{ "name": "NNERuntimeRDGHlsl" },
+						{ "name": "NNERuntimeRDGDml" }
 					],
 					"absolute_tolerance": 1e-5
 				},
@@ -46,8 +46,8 @@
 					"tags": [ "SmokeTest" ],
 					"runtimes": [
 						{ "name": "NNERuntimeORTDml", "skip": true, "comment": "" },
-						{ "name": "NNERuntimeHlsl",	"skip": true, "comment": "Model not supported yet."	},
-						{ "name": "NNERuntimeDml",  "skip": true, "comment": "Model not supported yet." }
+						{ "name": "NNERuntimeRDGHlsl",	"skip": true, "comment": "Model not supported yet."	},
+						{ "name": "NNERuntimeRDGDml",  "skip": true, "comment": "Model not supported yet." }
 					]
 				},
 				{
@@ -55,16 +55,16 @@
 					"tags": [ "SmokeTest" ],
 					"runtimes": [
 						{ "name": "NNERuntimeORTDml", "absolute_tolerance": 1e-6 },
-						{ "name": "NNERuntimeHlsl",  "skip": true, "comment": "Missing Reshape operator." },
-						{ "name": "NNERuntimeDml",  "skip": true, "comment": "Model not supported yet." }
+						{ "name": "NNERuntimeRDGHlsl",  "skip": true, "comment": "Missing Reshape operator." },
+						{ "name": "NNERuntimeRDGDml",  "skip": true, "comment": "Model not supported yet." }
 					]
 				},
 				{
 					"target": "VertexDeltaModel",
 					"runtimes": [
 						{ "name": "NNERuntimeORTDml", "absolute_tolerance": 1e-4 },
-						{ "name": "NNERuntimeHlsl", "absolute_tolerance": 1e-4 },
-						{ "name": "NNERuntimeDml",  "skip": true, "comment": "Model not supported yet." }
+						{ "name": "NNERuntimeRDGHlsl", "absolute_tolerance": 1e-4 },
+						{ "name": "NNERuntimeRDGDml",  "skip": true, "comment": "Model not supported yet." }
 					]
 				}
 			]
@@ -76,21 +76,21 @@
 				{
 					"target": "OneLayerWithWeights",
 					"runtimes": [
-						{ "name": "NNERuntimeDml",  "skip": true, "comment": "Model not supported yet." }
+						{ "name": "NNERuntimeRDGDml",  "skip": true, "comment": "Model not supported yet." }
 					]
 				},
 				{
 					"target": "TwoLayersNoWeight",
 					"tags": [ "SmokeTest" ],
 					"runtimes": [
-						{ "name": "NNERuntimeDml",  "skip": true, "comment": "Model not supported yet." }
+						{ "name": "NNERuntimeRDGDml",  "skip": true, "comment": "Model not supported yet." }
 					]
 				},
 				{
 					"target": "MultiOutput",
 					"tags": [ "SmokeTest" ],
 					"runtimes": [
-						{ "name": "NNERuntimeDml",  "skip": true, "comment": "Model not supported yet." }
+						{ "name": "NNERuntimeRDGDml",  "skip": true, "comment": "Model not supported yet." }
 					]
 				}
 			]
@@ -102,14 +102,14 @@
 				{
 					"target": "OneLayerNoWeightVariableInputShape",
 					"runtimes": [
-						{ "name": "NNERuntimeDml",  "skip": true, "comment": "RDG does not support variable input sized tensor" }
+						{ "name": "NNERuntimeRDGDml",  "skip": true, "comment": "RDG does not support variable input sized tensor" }
 					]
 				},
 				{
 					"target": "TwoLayersNoWeightVariableInputShape",
 					"tags": [ "SmokeTest" ],
 					"runtimes": [
-						{ "name": "NNERuntimeDml",  "skip": true, "comment": "RDG does not support variable input sized tensor" }
+						{ "name": "NNERuntimeRDGDml",  "skip": true, "comment": "RDG does not support variable input sized tensor" }
 					]
 				}
 			]
@@ -123,76 +123,76 @@
 				{	"target": "Abs", "additional_datasets": [ "Elementwise.Slow.OneInput" ]},
 				{	"target": "Acos", "tags": [ "SmokeTest" ], "runtimes": [
 						{ "name": "NNERuntimeORTDml", "absolute_tolerance":  1e-4},
-						{ "name": "NNERuntimeDml",    "absolute_tolerance":  1e-4},
-						{ "name": "NNERuntimeHlsl",   "absolute_tolerance":  2e-4}
+						{ "name": "NNERuntimeRDGDml",    "absolute_tolerance":  1e-4},
+						{ "name": "NNERuntimeRDGHlsl",   "absolute_tolerance":  2e-4}
 					]},
 				{	"target": "Acosh" },
 				{	"target": "Asin", "runtimes": [
 						{ "name": "NNERuntimeORTDml", "absolute_tolerance":  1e-4},
-						{ "name": "NNERuntimeDml",    "absolute_tolerance":  1e-4},
-						{ "name": "NNERuntimeHlsl",   "absolute_tolerance":  2e-4}
+						{ "name": "NNERuntimeRDGDml",    "absolute_tolerance":  1e-4},
+						{ "name": "NNERuntimeRDGHlsl",   "absolute_tolerance":  2e-4}
 					]},
 				{	"target": "Asinh", "runtimes": [
 						{ "name": "NNERuntimeORTDml", "absolute_tolerance":  1e-4},
-						{ "name": "NNERuntimeDml",    "absolute_tolerance":  1e-4},
-						{ "name": "NNERuntimeHlsl",   "absolute_tolerance":  2e-4}
+						{ "name": "NNERuntimeRDGDml",    "absolute_tolerance":  1e-4},
+						{ "name": "NNERuntimeRDGHlsl",   "absolute_tolerance":  2e-4}
 					]},
 				{	"target": "Atan", "runtimes": [
 						{ "name": "NNERuntimeORTDml", "absolute_tolerance":  1e-4},
-						{ "name": "NNERuntimeDml",    "absolute_tolerance":  1e-4},
-						{ "name": "NNERuntimeHlsl",   "absolute_tolerance":  1.5e-3}
+						{ "name": "NNERuntimeRDGDml",    "absolute_tolerance":  1e-4},
+						{ "name": "NNERuntimeRDGHlsl",   "absolute_tolerance":  1.5e-3}
 					]},
 				{	"target": "Atanh" },
 				{	"target": "BitShift", "skip": true, "comment": "TODO need attributes" },
 				{	"target": "Cast",     "skip": true, "comment": "TODO need attributes" },
 				{	"target": "Ceil"},
 				{	"target": "Clip", "additional_datasets": [ "Elementwise.ClipTestWithExtraInput" ], "runtimes": [
-						{ "name": "NNERuntimeDml", "skip": true, "comment": "RDG need scalar tensor support"},
-						{ "name": "NNERuntimeHlsl", "skip": true, "comment": "RDG need scalar tensor support"}						
+						{ "name": "NNERuntimeRDGDml", "skip": true, "comment": "RDG need scalar tensor support"},
+						{ "name": "NNERuntimeRDGHlsl", "skip": true, "comment": "RDG need scalar tensor support"}						
 					]},
 				{	"target": "Cos" },
 				{	"target": "Cosh" },
 				{	"target": "Elu", "comment": "TODO test with not-default attribute", "runtimes": [
 						{ "name": "NNERuntimeORTDml", "absolute_tolerance":  1e-4},
-						{ "name": "NNERuntimeDml",    "absolute_tolerance":  1e-4},
-						{ "name": "NNERuntimeHlsl",   "absolute_tolerance":  1e-4}
+						{ "name": "NNERuntimeRDGDml",    "absolute_tolerance":  1e-4},
+						{ "name": "NNERuntimeRDGHlsl",   "absolute_tolerance":  1e-4}
 					]},
 				{	"target": "Erf", "runtimes": [
 						{ "name": "NNERuntimeORTDml", "absolute_tolerance":  1e-4},
-						{ "name": "NNERuntimeDml",    "absolute_tolerance":  1e-4},
-						{ "name": "NNERuntimeHlsl",   "absolute_tolerance":  5e-4, "comment": "NNERuntimeHlsl Erf implementation is less precise than the DML one"}
+						{ "name": "NNERuntimeRDGDml",    "absolute_tolerance":  1e-4},
+						{ "name": "NNERuntimeRDGHlsl",   "absolute_tolerance":  5e-4, "comment": "NNERuntimeRDGHlsl Erf implementation is less precise than the DML one"}
 					]},
 				{	"target": "Exp" },
 				{	"target": "Floor" },
 				{	"target": "IsInf", "skip": true, "comment": "TODO add flexibility in test data setup + bool tensors" },
 				{	"target": "IsNan", "skip": true, "comment": "TODO add flexibility in test data setup + bool tensors" },
 				{	"target": "HardSigmoid",         "comment": "TODO test with not-default attribute" },
-				{	"target": "HardSwish", "runtimes": [ { "name": "NNERuntimeDml", "skip": true, "comment": "Not implemented atm on NNERuntimeDml."} ]},
+				{	"target": "HardSwish", "runtimes": [ { "name": "NNERuntimeRDGDml", "skip": true, "comment": "Not implemented atm on NNERuntimeRDGDml."} ]},
 				{	"target": "LeakyRelu", "additional_datasets": [ "Elementwise.LeakyRelu.AttributeTest" ] },
 				{	"target": "Log", "runtimes": [
 						{ "name": "NNERuntimeORTDml"},
-						{ "name": "NNERuntimeDml"},
-						{ "name": "NNERuntimeHlsl"}
+						{ "name": "NNERuntimeRDGDml"},
+						{ "name": "NNERuntimeRDGHlsl"}
 					], "absolute_tolerance": 1e-6 },
 				{	"target": "Neg" },
 				{	"target": "Not", "input_type": "boolean", "output_type": "boolean", "runtimes": [
-						{ "name": "NNERuntimeDml", "skip": true, "comment": "RDG need bool tensor support"},
-						{ "name": "NNERuntimeHlsl", "skip": true, "comment": "RDG need bool tensor support"}						
+						{ "name": "NNERuntimeRDGDml", "skip": true, "comment": "RDG need bool tensor support"},
+						{ "name": "NNERuntimeRDGHlsl", "skip": true, "comment": "RDG need bool tensor support"}						
 					]},
 				{	"target": "Reciprocal"},
 				{	"target": "Relu"},
 				{	"target": "Round"},
 				{	"target": "Selu", "comment": "TODO test with not-default attribute", "runtimes": [
 						{ "name": "NNERuntimeORTDml", "absolute_tolerance":  1e-4},
-						{ "name": "NNERuntimeDml",    "absolute_tolerance":  1e-4},
-						{ "name": "NNERuntimeHlsl",   "absolute_tolerance":  1e-4}
+						{ "name": "NNERuntimeRDGDml",    "absolute_tolerance":  1e-4},
+						{ "name": "NNERuntimeRDGHlsl",   "absolute_tolerance":  1e-4}
 					]},
 				{	"target": "Sigmoid"},
 				{	"target": "Sign" },
 				{	"target": "Sin", "runtimes": [
 						{ "name": "NNERuntimeORTDml", "absolute_tolerance":  1e-4},
-						{ "name": "NNERuntimeDml",    "absolute_tolerance":  1e-4},
-						{ "name": "NNERuntimeHlsl",   "absolute_tolerance":  1e-4}
+						{ "name": "NNERuntimeRDGDml",    "absolute_tolerance":  1e-4},
+						{ "name": "NNERuntimeRDGHlsl",   "absolute_tolerance":  1e-4}
 					]},
 				{	"target": "Sinh" },
 				{	"target": "Softplus"},
@@ -200,13 +200,13 @@
 				{	"target": "Sqrt"},
 				{	"target": "Tan", "runtimes": [
 						{ "name": "NNERuntimeORTDml", "absolute_tolerance":  1e-4},
-						{ "name": "NNERuntimeDml",    "absolute_tolerance":  1e-4},
-						{ "name": "NNERuntimeHlsl",   "absolute_tolerance":  1e-4}
+						{ "name": "NNERuntimeRDGDml",    "absolute_tolerance":  1e-4},
+						{ "name": "NNERuntimeRDGHlsl",   "absolute_tolerance":  1e-4}
 					]},
 				{	"target": "Tanh", "runtimes": [
 						{ "name": "NNERuntimeORTDml", "absolute_tolerance":  1e-4},
-						{ "name": "NNERuntimeDml",    "absolute_tolerance":  1e-4},
-						{ "name": "NNERuntimeHlsl",   "absolute_tolerance":  1e-4}
+						{ "name": "NNERuntimeRDGDml",    "absolute_tolerance":  1e-4},
+						{ "name": "NNERuntimeRDGHlsl",   "absolute_tolerance":  1e-4}
 					]}
 			]
 		},
@@ -215,35 +215,35 @@
 			"targets": [
 				{	"target": "Add", "additional_datasets": [ "Elementwise.Slow.TwoInputs" ]},
 				{	"target": "And", "input_type": "boolean", "output_type": "boolean", "runtimes": [
-						{ "name": "NNERuntimeDml", "skip": true, "comment": "RDG need bool tensor support"},
-						{ "name": "NNERuntimeHlsl", "skip": true, "comment": "RDG need bool tensor support"}						
+						{ "name": "NNERuntimeRDGDml", "skip": true, "comment": "RDG need bool tensor support"},
+						{ "name": "NNERuntimeRDGHlsl", "skip": true, "comment": "RDG need bool tensor support"}						
 					]},
 				{	"target": "Div", "tags": [ "SmokeTest" ]},
 				{	"target": "Equal", "output_type": "boolean", "runtimes": [
-						{ "name": "NNERuntimeDml", "skip": true, "comment": "RDG need bool tensor support"},
-						{ "name": "NNERuntimeHlsl", "skip": true, "comment": "RDG need bool tensor support"}						
+						{ "name": "NNERuntimeRDGDml", "skip": true, "comment": "RDG need bool tensor support"},
+						{ "name": "NNERuntimeRDGHlsl", "skip": true, "comment": "RDG need bool tensor support"}						
 					]},
 				{	"target": "Greater", "output_type": "boolean", "runtimes": [
-						{ "name": "NNERuntimeDml", "skip": true, "comment": "RDG need bool tensor support"},
-						{ "name": "NNERuntimeHlsl", "skip": true, "comment": "RDG need bool tensor support"}						
+						{ "name": "NNERuntimeRDGDml", "skip": true, "comment": "RDG need bool tensor support"},
+						{ "name": "NNERuntimeRDGHlsl", "skip": true, "comment": "RDG need bool tensor support"}						
 					]},
 				{	"target": "GreaterOrEqual", "output_type": "boolean", "runtimes": [
-						{ "name": "NNERuntimeDml", "skip": true, "comment": "RDG need bool tensor support"},
-						{ "name": "NNERuntimeHlsl", "skip": true, "comment": "RDG need bool tensor support"}						
+						{ "name": "NNERuntimeRDGDml", "skip": true, "comment": "RDG need bool tensor support"},
+						{ "name": "NNERuntimeRDGHlsl", "skip": true, "comment": "RDG need bool tensor support"}						
 					]},
 				{	"target": "Less", "output_type": "boolean", "runtimes": [
-						{ "name": "NNERuntimeDml", "skip": true, "comment": "RDG need bool tensor support"},
-						{ "name": "NNERuntimeHlsl", "skip": true, "comment": "RDG need bool tensor support"}						
+						{ "name": "NNERuntimeRDGDml", "skip": true, "comment": "RDG need bool tensor support"},
+						{ "name": "NNERuntimeRDGHlsl", "skip": true, "comment": "RDG need bool tensor support"}						
 					]},
 				{	"target": "LessOrEqual", "output_type": "boolean", "runtimes": [
-						{ "name": "NNERuntimeDml", "skip": true, "comment": "RDG need bool tensor support"},
-						{ "name": "NNERuntimeHlsl", "skip": true, "comment": "RDG need bool tensor support"}						
+						{ "name": "NNERuntimeRDGDml", "skip": true, "comment": "RDG need bool tensor support"},
+						{ "name": "NNERuntimeRDGHlsl", "skip": true, "comment": "RDG need bool tensor support"}						
 					]},
 				{	"target": "Mod", "skip": true, "comment": "TODO need to set fmod attribute to 1 for floating point tensors"},
 				{	"target": "Mul"},
 				{	"target": "Or", "input_type": "boolean", "output_type": "boolean", "runtimes": [
-						{ "name": "NNERuntimeDml", "skip": true, "comment": "RDG need bool tensor support"},
-						{ "name": "NNERuntimeHlsl", "skip": true, "comment": "RDG need bool tensor support"}						
+						{ "name": "NNERuntimeRDGDml", "skip": true, "comment": "RDG need bool tensor support"},
+						{ "name": "NNERuntimeRDGHlsl", "skip": true, "comment": "RDG need bool tensor support"}						
 					]},
 				{	"target": "PRelu", "removed_datasets": [ 
 						"Elementwise.Binary.BidirectionalBroadcast", 
@@ -252,15 +252,15 @@
 				{	"target": "Pow", "absolute_tolerance": 1e-4},
 				{	"target": "Sub"},
 				{	"target": "Xor", "input_type": "boolean", "output_type": "boolean", "runtimes": [
-						{ "name": "NNERuntimeDml", "skip": true, "comment": "RDG need bool tensor support"},
-						{ "name": "NNERuntimeHlsl", "skip": true, "comment": "RDG need bool tensor support"}						
+						{ "name": "NNERuntimeRDGDml", "skip": true, "comment": "RDG need bool tensor support"},
+						{ "name": "NNERuntimeRDGHlsl", "skip": true, "comment": "RDG need bool tensor support"}						
 					]}
 			]
 		},
 		{
 			"category": "Elementwise.Variadic",
 			"additional_datasets": [ "Elementwise.Binary" ],
-			"runtimes": [ { "name": "NNERuntimeDml", "skip": true} ],
+			"runtimes": [ { "name": "NNERuntimeRDGDml", "skip": true} ],
 			"targets": [
 				{	"target": "Max", "additional_datasets": [ "Elementwise.Slow.TwoInputs" ]},
 				{	"target": "Mean", "tags": [ "SmokeTest" ]},
@@ -277,9 +277,9 @@
 		{
 			"category": "Convolution",
 			"runtimes": [
-				{ "name": "NNERuntimeDml", "skip": true	},
+				{ "name": "NNERuntimeRDGDml", "skip": true	},
 				{ "name": "NNERuntimeORTDml" },
-				{ "name": "NNERuntimeHlsl" }
+				{ "name": "NNERuntimeRDGHlsl" }
 			],
 			"targets": [
 				{ "target": "Conv", "absolute_tolerance": 1e-6 }
@@ -288,9 +288,9 @@
 		{
 			"category": "ConvTranspose",
 			"runtimes": [
-				{ "name": "NNERuntimeDml", "skip": true	},
+				{ "name": "NNERuntimeRDGDml", "skip": true	},
 				{ "name": "NNERuntimeORTDml" },
-				{ "name": "NNERuntimeHlsl" }
+				{ "name": "NNERuntimeRDGHlsl" }
 			],
 			"targets": [
 				{ "target": "ConvTranspose", "absolute_tolerance": 1e-6 }
@@ -298,7 +298,7 @@
 		},
 		{
 			"category": "MatMul",
-			"runtimes": [ { "name": "NNERuntimeDml", "skip": true } ],
+			"runtimes": [ { "name": "NNERuntimeRDGDml", "skip": true } ],
 			"targets": [
 				{ "target": "MatMul" }
 			]
@@ -306,7 +306,7 @@
 		{
 			"category": "PadOp",
 			"runtimes": [
-				{ "name": "NNERuntimeDml", "skip": true	}
+				{ "name": "NNERuntimeRDGDml", "skip": true	}
 			],
 			"targets": [
 				{ "target": "Pad" }
@@ -315,7 +315,7 @@
 		{
 			"category": "InstanceNormalization",
 			"runtimes": [
-				{ "name": "NNERuntimeDml", "skip": true }
+				{ "name": "NNERuntimeRDGDml", "skip": true }
 			],
 			"targets": [
 				{ "target": "InstanceNormalization", "absolute_tolerance": 1e-5 }
@@ -324,7 +324,7 @@
 		{
 			"category": "Upsample",
 			"runtimes": [
-				{ "name": "NNERuntimeDml", "skip": true	}
+				{ "name": "NNERuntimeRDGDml", "skip": true	}
 			],
 			"targets": [
 				{ "target": "Upsample" }
@@ -353,8 +353,8 @@
 			"datasets": [
 				{ "inputs": [ { "shape": [ 1,0,3 ] } ],
 					"runtimes": [
-						{ "name": "NNERuntimeHlsl", "skip": true, "comment": "RDG does not support 0 sized tensor" },
-						{ "name": "NNERuntimeDml",  "skip": true, "comment": "RDG does not support 0 sized tensor" }
+						{ "name": "NNERuntimeRDGHlsl", "skip": true, "comment": "RDG does not support 0 sized tensor" },
+						{ "name": "NNERuntimeRDGDml",  "skip": true, "comment": "RDG does not support 0 sized tensor" }
 					]}
 			]
 		},
@@ -403,8 +403,8 @@
 			"datasets": [
 				{ "inputs": [ { "shape": [ 1,0 ] }, { "shape": [ 1,1 ] } ],
 					"runtimes": [
-						{ "name": "NNERuntimeHlsl", "skip": true, "comment": "RDG does not support 0 sized tensor" },
-						{ "name": "NNERuntimeDml",  "skip": true, "comment": "RDG does not support 0 sized tensor" }
+						{ "name": "NNERuntimeRDGHlsl", "skip": true, "comment": "RDG does not support 0 sized tensor" },
+						{ "name": "NNERuntimeRDGDml",  "skip": true, "comment": "RDG does not support 0 sized tensor" }
 					]}
 			]
 		},
@@ -562,7 +562,7 @@
 			"datasets": [ { "inputs": [ { "shape": [ 1, 1, 3, 3 ] }, { "shape": [ 1, 2, 3, 3 ] } ], "outputs": [ { "shape": [ 1, 2, 6, 6 ] } ] } ],
 			"runtimes": [
 				{ "name": "NNERuntimeORTDml", "skip": true, "comment": "TODO INVESTIGATE Returns 'The parameter is incorrect'" },
-				{ "name": "NNERuntimeHlsl", "skip": true, "comment": "Fails until ORT 1.13, because they implemented 'auto_pad' differently (they flipped 'SAME_LOWER' and 'SAME_UPPER')" }
+				{ "name": "NNERuntimeRDGHlsl", "skip": true, "comment": "Fails until ORT 1.13, because they implemented 'auto_pad' differently (they flipped 'SAME_LOWER' and 'SAME_UPPER')" }
 			]
 		},
 		{
@@ -584,7 +584,7 @@
 			"name": "MatMul.Rank5",
 			"runtimes": [	
 				{ "name": "NNERuntimeORTDml", "skip": true, "comment": "DML metacommand only support 2 to 4 for inputs rank" },
-				{ "name": "NNERuntimeHlsl"}
+				{ "name": "NNERuntimeRDGHlsl"}
 			],
 			"datasets": [
 				{ "inputs": [ { "shape": [ 2, 2, 3, 5, 4 ] }, { "shape": [ 2, 1, 3, 4, 3 ] } ], "outputs": [ { "shape": [ 2, 2, 3, 5, 3 ] } ]},
@@ -628,8 +628,8 @@
 		{
 			"name": "Upsample.ScaleFromInput",
 			"runtimes": [	
-				{ "name": "NNERuntimeDml", "skip": true, "comment": "Scale should be a constant tensor for RDG" },
-				{ "name": "NNERuntimeHlsl", "skip": true, "comment": "Scale should be a constant tensor for RDG" }
+				{ "name": "NNERuntimeRDGDml", "skip": true, "comment": "Scale should be a constant tensor for RDG" },
+				{ "name": "NNERuntimeRDGHlsl", "skip": true, "comment": "Scale should be a constant tensor for RDG" }
 			],
 			"datasets": [
 				{ "inputs":  [ { "shape": [ 1 ] },                            { "shape": [ 1 ], "source": [ "3"    ],                  "type": "float" } ], "outputs": [ { "shape": [ 3 ]} ] },
@@ -661,64 +661,64 @@
 				{
 					"inputs": [ { "shape": [ 2, 2, 4, 1, 3 ] }, { "shape": [ 2 ] }, { "shape": [ 2 ] } ], "outputs": [ { "shape": [ 2, 2, 4, 1, 3 ] } ],
 					"runtimes": [
-						{ "name": "NNERuntimeHlsl" },
-						{ "name": "NNERuntimeORTCpuEditor" },
+						{ "name": "NNERuntimeRDGHlsl" },
+						{ "name": "NNERuntimeORTCpu" },
 						{ "name": "NNERuntimeORTDml", "skip": true, "comment": "Not supported" }
 					]
 				},
 				{
 					"inputs": [ { "shape": [ 2, 2, 4, 4 ] }, { "shape": [ 2 ] }, { "shape": [ 2 ] } ], "outputs": [ { "shape": [ 2, 2, 4, 4 ] } ],
 					"runtimes": [
-						{ "name": "NNERuntimeHlsl" },
-						{ "name": "NNERuntimeORTCpuEditor" },
+						{ "name": "NNERuntimeRDGHlsl" },
+						{ "name": "NNERuntimeORTCpu" },
 						{ "name": "NNERuntimeORTDml" }
 					]
 				},
 				{
 					"inputs": [ { "shape": [ 2, 2, 8, 8 ] }, { "shape": [ 2 ] }, { "shape": [ 2 ] } ], "outputs": [ { "shape": [ 2, 2, 8, 8 ] } ],
 					"runtimes": [
-						{ "name": "NNERuntimeHlsl" },
-						{ "name": "NNERuntimeORTCpuEditor" },
+						{ "name": "NNERuntimeRDGHlsl" },
+						{ "name": "NNERuntimeORTCpu" },
 						{ "name": "NNERuntimeORTDml" }
 					]
 				},
 				{
 					"inputs": [ { "shape": [ 2, 128, 8, 8 ] }, { "shape": [ 128 ] }, { "shape": [ 128 ] } ], "outputs": [ { "shape": [ 2, 128, 8, 8 ] } ],
 					"runtimes": [
-						{ "name": "NNERuntimeHlsl" },
-						{ "name": "NNERuntimeORTCpuEditor" },
+						{ "name": "NNERuntimeRDGHlsl" },
+						{ "name": "NNERuntimeORTCpu" },
 						{ "name": "NNERuntimeORTDml" }
 					]
 				},
 				{
 					"inputs": [ { "shape": [ 128, 2, 8, 8 ] }, { "shape": [ 2 ] }, { "shape": [ 2 ] } ], "outputs": [ { "shape": [ 128, 2, 8, 8 ] } ],
 					"runtimes": [
-						{ "name": "NNERuntimeHlsl" },
-						{ "name": "NNERuntimeORTCpuEditor" },
+						{ "name": "NNERuntimeRDGHlsl" },
+						{ "name": "NNERuntimeORTCpu" },
 						{ "name": "NNERuntimeORTDml" }
 					]
 				},
 				{
 					"inputs": [ { "shape": [ 64, 64, 8, 8 ] }, { "shape": [ 64 ] }, { "shape": [ 64 ] } ], "outputs": [ { "shape": [ 64, 64, 8, 8 ] } ],
 					"runtimes": [
-						{ "name": "NNERuntimeHlsl" },
-						{ "name": "NNERuntimeORTCpuEditor" },
+						{ "name": "NNERuntimeRDGHlsl" },
+						{ "name": "NNERuntimeORTCpu" },
 						{ "name": "NNERuntimeORTDml" }
 					]
 				},
 				{
 					"inputs": [ { "shape": [ 2, 2, 45, 45 ] }, { "shape": [ 2 ] }, { "shape": [ 2 ] } ], "outputs": [ { "shape": [ 2, 2, 45, 45 ] } ],
 					"runtimes": [
-						{ "name": "NNERuntimeHlsl" },
-						{ "name": "NNERuntimeORTCpuEditor" },
+						{ "name": "NNERuntimeRDGHlsl" },
+						{ "name": "NNERuntimeORTCpu" },
 						{ "name": "NNERuntimeORTDml" }
 					]
 				},
 				{
 					"inputs": [ { "shape": [ 1, 128, 56, 56 ] }, { "shape": [ 128 ] }, { "shape": [ 128 ] } ], "outputs": [ { "shape": [ 1, 128, 56, 56 ] } ],
 					"runtimes": [
-						{ "name": "NNERuntimeHlsl" },
-						{ "name": "NNERuntimeORTCpuEditor" },
+						{ "name": "NNERuntimeRDGHlsl" },
+						{ "name": "NNERuntimeORTCpu" },
 						{ "name": "NNERuntimeORTDml" }
 					],
 					"comment": "from 'mosaic-9-sim.onnx'"
