NNE - Fix HLSL runtime duplicating the model weights on CPU.

Problem: HLSL runtime was duplicating weights on CPU.
Solution: use FTensorRDG for HLSL runtime as we dont need up/download hlsl tensor feature anymore
Bonus: fix a bug about FModel internal state creation were tensorRef could point to invalid memory.

Test: all nne test green on editor win64 + style transfer

diff --git a/Engine/Plugins/Experimental/NNE/Source/NNECore/Internal/NNECoreModelBase.h b/Engine/Plugins/Experimental/NNE/Source/NNECore/Internal/NNECoreModelBase.h
--- a/Engine/Plugins/Experimental/NNE/Source/NNECore/Internal/NNECoreModelBase.h
+++ b/Engine/Plugins/Experimental/NNE/Source/NNECore/Internal/NNECoreModelBase.h
@@ -52,7 +52,7 @@ namespace UE::NNECore::Internal
 
 	template <class T> int FModelBase<T>::SetInputTensorShapes(TConstArrayView<FTensorShape> InInputShapes)
 	{
-		InputTensorShapes.Empty();
+		InputTensorShapes.Reset(InInputShapes.Num());
 
 		if (InInputShapes.Num() != InputSymbolicTensors.Num())
 		{
diff --git a/Engine/Plugins/Experimental/NNE/Source/NNECore/Internal/NNECoreTensor.h b/Engine/Plugins/Experimental/NNE/Source/NNECore/Internal/NNECoreTensor.h
--- a/Engine/Plugins/Experimental/NNE/Source/NNECore/Internal/NNECoreTensor.h
+++ b/Engine/Plugins/Experimental/NNE/Source/NNECore/Internal/NNECoreTensor.h
@@ -62,7 +62,7 @@ namespace UE::NNECore::Internal
 		void SetPreparedData(TConstArrayView<uint8> Data)
 		{
 			checkf(Data.Num() == DataSize, TEXT("Incorrect data size, it should match tensor shape and data type."));
-			PreparedData.Empty();
+			PreparedData.Reset();
 			PreparedData.Append(Data);
 		}
 
diff --git a/Engine/Plugins/Experimental/NNE/Source/NNECore/Private/NNECoreModelData.cpp b/Engine/Plugins/Experimental/NNE/Source/NNECore/Private/NNECoreModelData.cpp
--- a/Engine/Plugins/Experimental/NNE/Source/NNECore/Private/NNECoreModelData.cpp
+++ b/Engine/Plugins/Experimental/NNE/Source/NNECore/Private/NNECoreModelData.cpp
@@ -137,7 +137,7 @@ void UNNEModelData::Serialize(FArchive& Ar)
 	// Recreate each model data when cooking
 	if (Ar.IsCooking() && Ar.IsSaving())
 	{
-		ModelData.Empty();
+		ModelData.Reset();
 		
 		TArrayView<TWeakInterfacePtr<INNERuntime>> Runtimes = UE::NNECore::GetAllRuntimes();
 		for (int i = 0; i < Runtimes.Num(); i++)
diff --git a/Engine/Plugins/Experimental/NNE/Source/NNEQA/Private/NNEQAUtils.cpp b/Engine/Plugins/Experimental/NNE/Source/NNEQA/Private/NNEQAUtils.cpp
--- a/Engine/Plugins/Experimental/NNE/Source/NNEQA/Private/NNEQAUtils.cpp
+++ b/Engine/Plugins/Experimental/NNE/Source/NNEQA/Private/NNEQAUtils.cpp
@@ -119,8 +119,8 @@ namespace UE::NNEQA::Private
 	static void FillInputTensorBindingsCPU(TConstArrayView<NNECore::Internal::FTensor> Tensors, TConstArrayView<FTests::FTensorData> TensorsData,
 		TArray<TArray<char>>& OutMemBuffers, TArray<NNECore::FTensorBindingCPU>& OutBindings)
 	{
-		OutBindings.Empty();
-		OutMemBuffers.Empty();
+		OutBindings.Reset();
+		OutMemBuffers.Reset();
 		check(TensorsData.Num() == 0 || Tensors.Num() == TensorsData.Num());
 
 		for (int Index = 0; Index < Tensors.Num(); ++Index)
@@ -144,8 +144,8 @@ namespace UE::NNEQA::Private
 
 	static void FillOutputTensorBindingsCPU(TConstArrayView<NNECore::Internal::FTensor> Tensors, TArray<TArray<char>>& OutMemBuffers, TArray<NNECore::FTensorBindingCPU>& OutBindings)
 	{
-		OutBindings.Empty();
-		OutMemBuffers.Empty();
+		OutBindings.Reset();
+		OutMemBuffers.Reset();
 		for (int Index = 0; Index < Tensors.Num(); ++Index)
 		{
 			TArray<char>& MemBuffer = OutMemBuffers.Emplace_GetRef();
@@ -381,7 +381,7 @@ namespace UE::NNEQA::Private
 	bool RunTestInference(const FNNEModelRaw& ONNXModelData, const FTests::FTestSetup& TestSetup,
 		const FString& RuntimeName, TArray<NNECore::Internal::FTensor>& OutOutputTensors, TArray<TArray<char>>& OutOutputMemBuffers)
 	{
-		OutOutputMemBuffers.Empty();
+		OutOutputMemBuffers.Reset();
 
 		TUniquePtr<FModelQA> InferenceModel = FModelQA::MakeModelQA(ONNXModelData, RuntimeName);
 		if (!InferenceModel.IsValid())
diff --git a/Engine/Plugins/Experimental/NNE/Source/NNERuntimeORT/Private/NNERuntimeORTModel.cpp b/Engine/Plugins/Experimental/NNE/Source/NNERuntimeORT/Private/NNERuntimeORTModel.cpp
--- a/Engine/Plugins/Experimental/NNE/Source/NNERuntimeORT/Private/NNERuntimeORTModel.cpp
+++ b/Engine/Plugins/Experimental/NNE/Source/NNERuntimeORT/Private/NNERuntimeORTModel.cpp
@@ -187,9 +187,9 @@ namespace UE::NNERuntimeORT::Private
 	
 	int FModelORT::SetInputTensorShapes(TConstArrayView<NNECore::FTensorShape> InInputShapes)
 	{
-		InputTensors.Empty();
-		OutputTensors.Empty();
-		OutputTensorShapes.Empty();
+		InputTensors.Reset();
+		OutputTensors.Reset();
+		OutputTensorShapes.Reset();
 
 		// Verify input shape are valid for the model and set InputTensorShapes
 		if (FModelBase<IModelCPU>::SetInputTensorShapes(InInputShapes) != 0)
@@ -219,8 +219,8 @@ namespace UE::NNERuntimeORT::Private
 		}
 		if (OutputTensors.Num() != OutputSymbolicTensors.Num())
 		{
-			OutputTensors.Empty();
-			OutputTensorShapes.Empty();
+			OutputTensors.Reset();
+			OutputTensorShapes.Reset();
 		}
 
 		return 0;
diff --git a/Engine/Plugins/Experimental/NNE/Source/NNERuntimeORTCpu/Private/NNERuntimeORTCpuModel.cpp b/Engine/Plugins/Experimental/NNE/Source/NNERuntimeORTCpu/Private/NNERuntimeORTCpuModel.cpp
--- a/Engine/Plugins/Experimental/NNE/Source/NNERuntimeORTCpu/Private/NNERuntimeORTCpuModel.cpp
+++ b/Engine/Plugins/Experimental/NNE/Source/NNERuntimeORTCpu/Private/NNERuntimeORTCpuModel.cpp
@@ -160,9 +160,9 @@ namespace UE::NNERuntimeORTCpu::Private
 	
 	int FModelCPU::SetInputTensorShapes(TConstArrayView<NNECore::FTensorShape> InInputShapes)
 	{
-		InputTensors.Empty();
-		OutputTensors.Empty();
-		OutputTensorShapes.Empty();
+		InputTensors.Reset();
+		OutputTensors.Reset();
+		OutputTensorShapes.Reset();
 
 		// Verify input shape are valid for the model and set InputTensorShapes
 		if (FModelBase<IModelCPU>::SetInputTensorShapes(InInputShapes) != 0)
@@ -192,8 +192,8 @@ namespace UE::NNERuntimeORTCpu::Private
 		}
 		if (OutputTensors.Num() != OutputSymbolicTensors.Num())
 		{
-			OutputTensors.Empty();
-			OutputTensorShapes.Empty();
+			OutputTensors.Reset();
+			OutputTensorShapes.Reset();
 		}
 
 		return 0;
diff --git a/Engine/Plugins/Experimental/NNE/Source/NNERuntimeRDG/Private/NNERuntimeRDGDml.cpp b/Engine/Plugins/Experimental/NNE/Source/NNERuntimeRDG/Private/NNERuntimeRDGDml.cpp
--- a/Engine/Plugins/Experimental/NNE/Source/NNERuntimeRDG/Private/NNERuntimeRDGDml.cpp
+++ b/Engine/Plugins/Experimental/NNE/Source/NNERuntimeRDG/Private/NNERuntimeRDGDml.cpp
@@ -1898,7 +1898,7 @@ void FModel::AddDispatchOps_RenderThread(FRDGBuilder& GraphBuilder)
 	for (int32 Idx = 0; Idx < InputTensorIndices.Num(); ++Idx)
 	{
 		FTensorBufferParamsDml* Params = GraphBuilder.AllocParameters<FTensorBufferParamsDml>();
-		Params->Buffer = AllTensorRDGs[InputTensorIndices[Idx]]->GetBuffer();
+		Params->Buffer = AllTensorRDGRefs[InputTensorIndices[Idx]]->GetBuffer();
 
 		GraphBuilder.AddPass(
 			RDG_EVENT_NAME("FModel_Dispatch_GetInputBuffer"),
@@ -1919,7 +1919,7 @@ void FModel::AddDispatchOps_RenderThread(FRDGBuilder& GraphBuilder)
 	for (int32 Idx = 0; Idx < OutputTensorIndices.Num(); ++Idx)
 	{
 		FTensorBufferParamsDml* Params = GraphBuilder.AllocParameters<FTensorBufferParamsDml>();
-		Params->Buffer = AllTensorRDGs[OutputTensorIndices[Idx]]->GetBuffer();
+		Params->Buffer = AllTensorRDGRefs[OutputTensorIndices[Idx]]->GetBuffer();
 
 		GraphBuilder.AddPass(
 			RDG_EVENT_NAME("FModel_Dispatch_GetOutputBuffer"),
diff --git a/Engine/Plugins/Experimental/NNE/Source/NNERuntimeRDG/Private/NNERuntimeRDGModel.cpp b/Engine/Plugins/Experimental/NNE/Source/NNERuntimeRDG/Private/NNERuntimeRDGModel.cpp
--- a/Engine/Plugins/Experimental/NNE/Source/NNERuntimeRDG/Private/NNERuntimeRDGModel.cpp
+++ b/Engine/Plugins/Experimental/NNE/Source/NNERuntimeRDG/Private/NNERuntimeRDGModel.cpp
@@ -98,7 +98,7 @@ bool FModelRDG::LoadModel(TConstArrayView<uint8> ModelData, FNNERuntimeFormat& F
 
 int FModelRDG::SetInputTensorShapes(TConstArrayView<NNECore::FTensorShape> InInputShapes)
 {
-	OutputTensorShapes.Empty();
+	OutputTensorShapes.Reset(OutputTensorIndices.Num());
 
 	//Verify input shape are valid for the model and set InputTensorShapes
 	if (FModelBase<NNECore::IModelRDG>::SetInputTensorShapes(InInputShapes) != 0)
@@ -106,10 +106,10 @@ int FModelRDG::SetInputTensorShapes(TConstArrayView<NNECore::FTensorShape> InInp
 		return -1;
 	}
 
-	//Allocate and prime all AllTensorRDGs with concrete shapes defaulting variables dimension to 1 if needed
-	AllTensorRDGs.Init(nullptr, AllSymbolicTensorDescs.Num());
+	//Allocate and prime all AllTensorRDGRefs with concrete shapes defaulting variables dimension to 1 if needed
+	AllTensorRDGRefs.Init(nullptr, AllSymbolicTensorDescs.Num());
 
-	InputTensorRDGs.Empty();
+	InputTensorRDGs.Reset(InputTensorIndices.Num());
 	for (int32 i = 0; i < InputTensorIndices.Num(); ++i)
 	{
 		const int32 Idx = InputTensorIndices[i];
@@ -117,17 +117,17 @@ int FModelRDG::SetInputTensorShapes(TConstArrayView<NNECore::FTensorShape> InInp
 		const NNECore::FTensorShape& TensorShape = InputTensorShapes[i];
 
 		InputTensorRDGs.Emplace(FTensorRDG::Make(TensorDesc, TensorShape, nullptr));
-		AllTensorRDGs[Idx] = &InputTensorRDGs[i];
+		AllTensorRDGRefs[Idx] = &InputTensorRDGs[i];
 	}
 
 	for (int32 i = 0; i < WeightTensorIndices.Num(); ++i)
 	{
 		const int32 Idx = WeightTensorIndices[i];
 
-		AllTensorRDGs[Idx] = &WeightTensorRDGs[i];
+		AllTensorRDGRefs[Idx] = &WeightTensorRDGs[i];
 	}
 
-	IntermediateTensorRDGs.Empty();
+	IntermediateTensorRDGs.Reset(IntermediateTensorIndices.Num());
 	for (int32 i = 0; i < IntermediateTensorIndices.Num(); ++i)
 	{
 		const int32 Idx = IntermediateTensorIndices[i];
@@ -135,10 +135,10 @@ int FModelRDG::SetInputTensorShapes(TConstArrayView<NNECore::FTensorShape> InInp
 		const NNECore::FTensorShape TensorShape = NNECore::FTensorShape::MakeFromSymbolic(TensorDesc.GetShape());
 
 		IntermediateTensorRDGs.Emplace(FTensorRDG::Make(TensorDesc, TensorShape, nullptr));
-		AllTensorRDGs[Idx] = &IntermediateTensorRDGs[i];
+		AllTensorRDGRefs[Idx] = &IntermediateTensorRDGs[i];
 	}
 
-	OutputTensorRDGs.Empty();
+	OutputTensorRDGs.Reset(OutputTensorIndices.Num());
 	for (int32 i = 0; i < OutputTensorIndices.Num(); ++i)
 	{
 		const int32 Idx = OutputTensorIndices[i];
@@ -146,13 +146,13 @@ int FModelRDG::SetInputTensorShapes(TConstArrayView<NNECore::FTensorShape> InInp
 		const NNECore::FTensorShape TensorShape = NNECore::FTensorShape::MakeFromSymbolic(TensorDesc.GetShape());
 
 		OutputTensorRDGs.Emplace(FTensorRDG::Make(TensorDesc, TensorShape, nullptr));
-		AllTensorRDGs[Idx] = &OutputTensorRDGs[i];
+		AllTensorRDGRefs[Idx] = &OutputTensorRDGs[i];
 	}
 
 	checkCode(
-		for (int i = 0; i < AllTensorRDGs.Num(); ++i)
+		for (int i = 0; i < AllTensorRDGRefs.Num(); ++i)
 		{
-			checkf(AllTensorRDGs[i] != nullptr, TEXT("Tensor at index %d, was not allocated for model preparation."), i);
+			checkf(AllTensorRDGRefs[i] != nullptr, TEXT("Tensor at index %d, was not allocated for model preparation."), i);
 		};
 	);
 
@@ -163,25 +163,25 @@ int FModelRDG::SetInputTensorShapes(TConstArrayView<NNECore::FTensorShape> InInp
 	}
 
 	checkCode(
-		for (int i = 0; i < AllTensorRDGs.Num(); ++i)
+		for (int i = 0; i < AllTensorRDGRefs.Num(); ++i)
 		{
-			checkf(AllTensorRDGs[i] != nullptr, TEXT("Tensor at index %d, was not allocated after model preparation."), i);
-			checkf(AllTensorRDGs[i]->GetShape().IsCompatibleWith(AllSymbolicTensorDescs[i].GetShape()), TEXT("Tensor at index %d have a shape incompatible with model definition."), i);
+			checkf(AllTensorRDGRefs[i] != nullptr, TEXT("Tensor at index %d, was not allocated after model preparation."), i);
+			checkf(AllTensorRDGRefs[i]->GetShape().IsCompatibleWith(AllSymbolicTensorDescs[i].GetShape()), TEXT("Tensor at index %d have a shape incompatible with model definition."), i);
 		};
 	);
 
 	//Set OutputTensorShapes for the model from preparation result
 	for (int32 OutputIndices : OutputTensorIndices)
 	{
-		OutputTensorShapes.Emplace(AllTensorRDGs[OutputIndices]->GetShape());
+		OutputTensorShapes.Emplace(AllTensorRDGRefs[OutputIndices]->GetShape());
 	}
 
-	check(InputTensorIndices.Num() + OutputTensorIndices.Num() + WeightTensorIndices.Num() + IntermediateTensorIndices.Num() == AllTensorRDGs.Num());
+	check(InputTensorIndices.Num() + OutputTensorIndices.Num() + WeightTensorIndices.Num() + IntermediateTensorIndices.Num() == AllTensorRDGRefs.Num());
 	check(InputTensorShapes.Num() == InputSymbolicTensors.Num());
 	check(OutputTensorShapes.Num() == OutputSymbolicTensors.Num());
 	check(WeightTensorIndices.Num() == WeightTensorRDGs.Num());
-	check(AllTensorRDGs.Num() == AllSymbolicTensorDescs.Num());
-
+	check(AllTensorRDGRefs.Num() == AllSymbolicTensorDescs.Num());
+	
 	return 0;
 }
 
@@ -211,16 +211,16 @@ int FModelRDG::EnqueueRDG(FRDGBuilder& RDGBuilder, TConstArrayView<NNECore::FTen
 	}
 
 	Res = SetTensors(RDGBuilder, InputTensorRDGs, InInputBindings);
-	if (Res != 0)
+	if (Res != -1)
 	{
-		UE_LOG(LogNNE, Warning, TEXT("Invalid input tensor binding type for tensor index:%d"), Res);
+		UE_LOG(LogNNE, Warning, TEXT("Invalid buffer (was nullptr) for input tensor binding at index %d"), Res);
 		return -1;
 	}
 
 	Res = SetTensors(RDGBuilder, OutputTensorRDGs, InOutputBindings);
-	if (Res != 0)
+	if (Res != -1)
 	{
-		UE_LOG(LogNNE, Warning, TEXT("Invalid output tensor binding type for tensor index:%d"), Res);
+		UE_LOG(LogNNE, Warning, TEXT("Invalid buffer (was nullptr) for output tensor binding at index %d"), Res);
 		return -1;
 	}
 
@@ -233,14 +233,17 @@ int FModelRDG::EnqueueRDG(FRDGBuilder& RDGBuilder, TConstArrayView<NNECore::FTen
 		TensorRDG.SetBuffer(TensorBuffer);
 	}
 
-	//Note: DirectML uses RHI buffers instead of RDG buffers
-	//For now weights tensors are not uploaded to GPU thus GetBuffer will return nullptr for them.
-	//checkCode(for (const FTensorRDG* TensorRDG : AllTensorRDGs) { if (TensorRDG != nullptr) { check(TensorRDG->GetBuffer() != nullptr); } });
+	if (AddWeightsToRDGGraph(RDGBuilder))
+	{
+		//Note: DirectML uses RHI buffers instead of RDG buffers
+		//For now weights tensors are not uploaded to GPU thus GetBuffer will return nullptr for them.
+		checkCode(for (const FTensorRDG* TensorRDG : AllTensorRDGRefs) { if (TensorRDG != nullptr) { check(TensorRDG->GetBuffer() != nullptr); } });	
+	}
 
 	//Insert weights tensors
 	for (int32 i = 0; i < WeightTensorIndices.Num(); ++i)
 	{
-		AllTensorRDGs[WeightTensorIndices[i]] = &WeightTensorRDGs[i];
+		AllTensorRDGRefs[WeightTensorIndices[i]] = &WeightTensorRDGs[i];
 	}
 
 	// We can now dispatch operators
@@ -252,15 +255,19 @@ int FModelRDG::EnqueueRDG(FRDGBuilder& RDGBuilder, TConstArrayView<NNECore::FTen
 int FModelRDG::SetTensors(FRDGBuilder& GraphBuilder, FTensorRDGArray& InTensorRDGs, TConstArrayView<NNECore::FTensorBindingRDG> InBindings)
 {
 	check(InBindings.Num() == InTensorRDGs.Num());
-
+	
 	for (int32 Idx = 0; Idx < InBindings.Num(); ++Idx)
 	{
 		FTensorRDG& TensorRDG = InTensorRDGs[Idx];
 		const NNECore::FTensorBindingRDG& Binding = InBindings[Idx];
+		if (Binding.Buffer == nullptr)
+		{
+			return Idx;
+		}
 		TensorRDG.SetBuffer(Binding.Buffer);
 	}
 
-	return 0;
+	return -1;
 }
 
 }
\ No newline at end of file
diff --git a/Engine/Plugins/Experimental/NNE/Source/NNERuntimeRDG/Private/NNERuntimeRDGModel.h b/Engine/Plugins/Experimental/NNE/Source/NNERuntimeRDG/Private/NNERuntimeRDGModel.h
--- a/Engine/Plugins/Experimental/NNE/Source/NNERuntimeRDG/Private/NNERuntimeRDGModel.h
+++ b/Engine/Plugins/Experimental/NNE/Source/NNERuntimeRDG/Private/NNERuntimeRDGModel.h
@@ -25,6 +25,7 @@ namespace UE::NNERuntimeRDG::Private
 		int SetTensors(FRDGBuilder& GraphBuilder, FTensorRDGArray& InTensorRDGs, TConstArrayView<NNECore::FTensorBindingRDG> InBindings);
 
 		virtual int PrepareTensorShapesAndData() = 0;
+		virtual bool AddWeightsToRDGGraph(FRDGBuilder& RDGBuilder) { return false; }
 		virtual void AddDispatchOps_RenderThread(FRDGBuilder& GraphBuilder) = 0;
 
 		//Tensor descriptor
@@ -41,7 +42,7 @@ namespace UE::NNERuntimeRDG::Private
 		TArray<TArray<uint32>>		OperatorOutputTensorIndices;
 
 		//RDG Tensors
-		FTensorRDGRefArray			AllTensorRDGs;
+		FTensorRDGRefArray			AllTensorRDGRefs;
 		FTensorRDGArray				InputTensorRDGs;
 		FTensorRDGArray				OutputTensorRDGs;
 		FTensorRDGArray				IntermediateTensorRDGs;
diff --git a/Engine/Plugins/Experimental/NNE/Source/NNERuntimeRDG/Private/NNERuntimeRDGModelHlsl.cpp b/Engine/Plugins/Experimental/NNE/Source/NNERuntimeRDG/Private/NNERuntimeRDGModelHlsl.cpp
--- a/Engine/Plugins/Experimental/NNE/Source/NNERuntimeRDG/Private/NNERuntimeRDGModelHlsl.cpp
+++ b/Engine/Plugins/Experimental/NNE/Source/NNERuntimeRDG/Private/NNERuntimeRDGModelHlsl.cpp
@@ -35,70 +35,23 @@ FOperatorHlsl* OpCreate(const FString& OpName, TConstArrayView<NNECore::FTensorD
 	return Op;
 }
 
-template<class OperatorRef>
-void InternAddDispatchOps_RenderThread(FRDGBuilder& GraphBuilder, TConstArrayView<FTensorHLSLRef> AllTensorHLSLRefs, TConstArrayView<TArray<uint32>> OperatorInputTensorIndices,
-	TConstArrayView<TArray<uint32>> OperatorOutputTensorIndices, TArrayView<OperatorRef> Operators)
-{
-	static constexpr int32 MaxExpectedInput = 10;
-	TArray<FTensorRDGRef, TInlineAllocator<MaxExpectedInput>> InputTensors;
-
-	static constexpr int32 MaxExpectedOutput = 2;
-	TArray<FTensorRDGRef, TInlineAllocator<MaxExpectedOutput>> OutputTensors;
-
-	// Add passes for all operators
-	for (int32 Idx = 0; Idx < Operators.Num(); ++Idx)
-	{
-		InputTensors.Empty();
-		for (int32 i : OperatorInputTensorIndices[Idx])
-		{
-			InputTensors.Add(AllTensorHLSLRefs[i]);
-		}
-		OutputTensors.Empty();
-		for (int32 i : OperatorOutputTensorIndices[Idx])
-		{
-			OutputTensors.Add(AllTensorHLSLRefs[i]);
-		}
-
-		Operators[Idx]->Dispatch(GraphBuilder, InputTensors, OutputTensors);
-	}
-}
+} // namespace ModelUtils
 
-int ApplyBinding(TArray<FTensorHLSL>& OutTensorsHLSL, TConstArrayView<NNECore::FTensorBindingRDG> InBindings, bool bIsInput)
+bool FModel::AddWeightsToRDGGraph(FRDGBuilder& RDGBuilder)
 {
-	check(OutTensorsHLSL.Num() == InBindings.Num());
-	
-	for (int32 Idx = 0; Idx < InBindings.Num(); ++Idx)
-	{
-		const NNECore::FTensorBindingRDG& Binding = InBindings[Idx];
-		FTensorHLSL& Tensor = OutTensorsHLSL[Idx];
-		
-		if (Binding.Buffer == nullptr)
-		{
-			return -1;
-		}
-		Tensor.SetBuffer(Binding.Buffer);
-	}
-
-	return 0;
-}
-
-bool ApplyWeights(FRDGBuilder& GraphBuilder, TArray<FTensorHLSL>& OutTensorsHLSL, TArray<TRefCountPtr<FRDGPooledBuffer>> ExternalWeightsRDG)
-{
-	check(OutTensorsHLSL.Num() == ExternalWeightsRDG.Num());
+	check(WeightTensorRDGs.Num() == WeightsExternalRDGResources.Num());
 
-	for (int32 Idx = 0; Idx < ExternalWeightsRDG.Num(); ++Idx)
+	for (int32 Idx = 0; Idx < WeightsExternalRDGResources.Num(); ++Idx)
 	{
-		const TRefCountPtr<FRDGPooledBuffer>& PooledBuffer = ExternalWeightsRDG[Idx];
-		FTensorHLSL& Tensor = OutTensorsHLSL[Idx];
-		FRDGBufferRef Buffer = GraphBuilder.RegisterExternalBuffer(PooledBuffer);
+		const TRefCountPtr<FRDGPooledBuffer>& PooledBuffer = WeightsExternalRDGResources[Idx];
+		FTensorRDG& Tensor = WeightTensorRDGs[Idx];
+		FRDGBufferRef Buffer = RDGBuilder.RegisterExternalBuffer(PooledBuffer);
 		Tensor.SetBuffer(Buffer);
 	}
 
 	return true;
 }
 
-}
-
 bool FModel::Init(TConstArrayView<uint8> ModelData)
 {
 	check(ModelData.Num() > 0);
@@ -152,96 +105,35 @@ bool FModel::Init(TConstArrayView<uint8> ModelData)
 	return true;
 }
 
-int FModel::SetInputTensorShapes(TConstArrayView<NNECore::FTensorShape> InputShapes)
+void FModel::AddDispatchOps_RenderThread(FRDGBuilder& GraphBuilder)
 {
-	int Res = FModelRDG::SetInputTensorShapes(InputShapes);
-	if (Res < 0) return Res;
+	static constexpr int32 MaxExpectedInput = 10;
+	TArray<FTensorRDGRef, TInlineAllocator<MaxExpectedInput>> InputTensors;
 
-	auto ConvertTensors = [this] (TArray<FTensorHLSLRef> &TensorRefs, TArray<FTensorHLSL> &Tensors, const FTensorRDGArray &TensorRDGs, const TArray<int32>& TensorIndices)
-	{
-		Tensors.Empty();
-		Tensors.SetNum(TensorRDGs.Num());
+	static constexpr int32 MaxExpectedOutput = 2;
+	TArray<FTensorRDGRef, TInlineAllocator<MaxExpectedOutput>> OutputTensors;
 
-		for (int32 i = 0; i < TensorIndices.Num(); ++i)
+	// Add passes for all operators
+	for (int32 Idx = 0; Idx < Operators.Num(); ++Idx)
+	{
+		InputTensors.Reset(OperatorInputTensorIndices.Num());
+		for (int32 i : OperatorInputTensorIndices[Idx])
 		{
-			const int32 Idx = TensorIndices[i];
-			Tensors[i] = FTensorHLSL(TensorRDGs[i]);
-			AllTensorHLSLRefs[Idx]= &Tensors[i];
+			InputTensors.Add(AllTensorRDGRefs[i]);
 		}
-	};
-
-	AllTensorHLSLRefs.Empty();
-	AllTensorHLSLRefs.SetNumUninitialized(AllTensorRDGs.Num());
-	ConvertTensors(AllTensorHLSLRefs, InputTensorHLSLs, InputTensorRDGs, InputTensorIndices);
-	ConvertTensors(AllTensorHLSLRefs, OutputTensorHLSLs, OutputTensorRDGs, OutputTensorIndices);
-	ConvertTensors(AllTensorHLSLRefs, IntermediateTensorHLSLs, IntermediateTensorRDGs, IntermediateTensorIndices);
-	for (int32 i = 0; i < WeightTensorIndices.Num(); ++i)
-	{
-		const int32 Idx = WeightTensorIndices[i];
-		AllTensorHLSLRefs[Idx] = &WeightTensorHLSLs[i];
-	}
-
-	return 0;
-}
-
-int FModel::EnqueueRDG(FRDGBuilder& GraphBuilder, TConstArrayView<NNECore::FTensorBindingRDG> InInputBindings, TConstArrayView<NNECore::FTensorBindingRDG> InOutputBindings)
-{
-	check(IsInRenderingThread());
-
-	int Res;
-
-	// Verify the model inputs were prepared
-	if (InputTensorShapes.Num() == 0)
-	{
-		UE_LOG(LogNNE, Error, TEXT("EnqueueRDG(): Input shapes are not set, please call SetInputTensorShapes."));
-		return -1;
-	}
-
-	Res = ModelUtils::ApplyBinding(InputTensorHLSLs, InInputBindings, true);
-	if (Res != 0)
-	{
-		UE_LOG(LogNNE, Warning, TEXT("Invalid input tensor binding for tensor index:%d"), Res);
-		return -1;
-	}
-
-	Res = ModelUtils::ApplyBinding(OutputTensorHLSLs, InOutputBindings, false);
-	if (Res != 0)
-	{
-		UE_LOG(LogNNE, Warning, TEXT("Invalid output tensor binding for tensor index:%d"), Res);
-		return -1;
-	}
-
-	if (!ModelUtils::ApplyWeights(GraphBuilder, WeightTensorHLSLs, WeightsExternalRDGResources))
-	{
-		UE_LOG(LogNNE, Warning, TEXT("Could not register the weights for graph execution."));
-		return -1;
-	}
-
-	// Create RDG buffers where not yet present
-	for (FTensorHLSLRef Tensor : AllTensorHLSLRefs)
-	{
-		if (!Tensor->HasBuffer())
+		OutputTensors.Reset(OperatorOutputTensorIndices.Num());
+		for (int32 i : OperatorOutputTensorIndices[Idx])
 		{
-			FRDGBufferDesc BufferDesc = FRDGBufferDesc::CreateBufferDesc(Tensor->GetElemByteSize(), Tensor->GetVolume());
-			if (Tensor->HasDownloadBuffer())
-			{
-				BufferDesc.Usage = EBufferUsageFlags(BufferDesc.Usage | BUF_SourceCopy);
-			}
-
-			const FRDGBufferRef TensorBuffer = GraphBuilder.CreateBuffer(BufferDesc, *Tensor->GetName(), ERDGBufferFlags::None);
-			
-			Tensor->SetBuffer(TensorBuffer);
+			OutputTensors.Add(AllTensorRDGRefs[i]);
 		}
-	}
 
-	ModelUtils::InternAddDispatchOps_RenderThread<FOperatorHlsl*>(GraphBuilder, AllTensorHLSLRefs, OperatorInputTensorIndices, OperatorOutputTensorIndices, Operators);
-
-	return 0;
+		Operators[Idx]->Dispatch(GraphBuilder, InputTensors, OutputTensors);
+	}
 }
 
 int FModel::PrepareTensorShapesAndData()
 {
-	check(AllTensorRDGs.Num() == AllSymbolicTensorDescs.Num());
+	check(AllTensorRDGRefs.Num() == AllSymbolicTensorDescs.Num());
 	
 	if (Operators.Num() == 0)
 	{
@@ -272,19 +164,19 @@ int FModel::PrepareTensorShapesAndData()
 	// This loop could be abstracted to a different system as it apply on FTensorRef & IPrepareOperator witch are RDG agnostics.
 	for (int32 Idx = 0; Idx < Operators.Num(); ++Idx)
 	{
-		InputTensors.Empty();
-		OutputTensors.Empty();
+		InputTensors.Reset();
+		OutputTensors.Reset();
 
 		//Operator inputs
 		for (int32 i : OperatorInputTensorIndices[Idx])
 		{
 			checkf(AllInitializedTensors[i] == true, TEXT("Input tensor %d for operator %d should have been initialized."), i, Idx);
-			InputTensors.Emplace(AllTensorRDGs[i]);
+			InputTensors.Emplace(AllTensorRDGRefs[i]);
 		}
 		//Operator outputs
 		for (int32 i : OperatorOutputTensorIndices[Idx])
 		{
-			OutputTensors.Emplace(AllTensorRDGs[i]);
+			OutputTensors.Emplace(AllTensorRDGRefs[i]);
 			checkf(AllInitializedTensors[i] == false, TEXT("Output tensor %d for operator %d should not have been initialized yet."), i, Idx);
 			checkCode(AllInitializedTensors[i] = true);
 		}
@@ -296,7 +188,7 @@ int FModel::PrepareTensorShapesAndData()
 			//Operator could not prepare the output tensors, meaning we can't allocate
 			//output buffer before running the model. This runtime does not support this.
 			UE_LOG(LogNNE, Warning, TEXT("Could not deduce tensor shapes for this model during shape inference, HLSL runtime wont support the model as it need to precompute all shapes for performance reasons."));
-			AllTensorRDGs.Empty();
+			AllTensorRDGRefs.Reset(AllSymbolicTensorDescs.Num());
 			return -1;
 		}
 	}
@@ -315,22 +207,13 @@ bool FModel::PrepareWeights()
 {
 	if (!WeightsExternalRDGResources.IsEmpty())
 	{
-		check(WeightsExternalRDGResources.Num() == WeightTensorHLSLs.Num())
+		check(WeightsExternalRDGResources.Num() == WeightTensorRDGs.Num())
 		return true;
 	}
 
-	//Convert to HLSL tensor
-	//This copy the weight. To be improved.
-	check(WeightTensorHLSLs.IsEmpty());
-	WeightTensorHLSLs.SetNum(WeightTensorRDGs.Num());
-	for (int32 Idx = 0; Idx < WeightTensorRDGs.Num(); ++Idx)
-	{
-		WeightTensorHLSLs[Idx] = FTensorHLSL(WeightTensorRDGs[Idx]);
-	}
-
 	//Upload to GPU
 	FEvent* Signal = FGenericPlatformProcess::GetSynchEventFromPool(false);
-	WeightsExternalRDGResources.SetNum(WeightTensorHLSLs.Num());
+	WeightsExternalRDGResources.SetNum(WeightTensorRDGs.Num());
 
 	ENQUEUE_RENDER_COMMAND(FModel_PrepareWeights)
 	(
@@ -344,16 +227,18 @@ bool FModel::PrepareWeights()
 
 			FRDGBuilder	RDGBuilder(RHICmdList);
 
-			for (int32 i = 0; i < WeightTensorHLSLs.Num(); ++i)
+			for (int32 i = 0; i < WeightTensorRDGs.Num(); ++i)
 			{
-				FTensorHLSL& Tensor = WeightTensorHLSLs[i];
+				FTensorRDG& Tensor = WeightTensorRDGs[i];
 				check(!Tensor.HasBuffer());
 				check(Tensor.HasPreparedData());
 				FRDGBufferDesc BufferDesc = FRDGBufferDesc::CreateBufferDesc(Tensor.GetElemByteSize(), Tensor.GetVolume());
 				const FRDGBufferRef TransientRDGBuffer = RDGBuilder.CreateBuffer(BufferDesc, *Tensor.GetName(), ERDGBufferFlags::None);
+				const uint8* TensorData = Tensor.GetPreparedData<uint8>().GetData();
+				
 				WeightsExternalRDGResources[i] = RDGBuilder.ConvertToExternalBuffer(TransientRDGBuffer);
 				Tensor.SetBuffer(TransientRDGBuffer);
-				Tensor.EnqueueUploadRdg(RDGBuilder);
+				RDGBuilder.QueueBufferUpload(TransientRDGBuffer, TensorData, Tensor.GetDataSize(), ERDGInitialDataFlags::NoCopy);
 			}
 
 			RDGBuilder.Execute();
diff --git a/Engine/Plugins/Experimental/NNE/Source/NNERuntimeRDG/Private/NNERuntimeRDGModelHlsl.h b/Engine/Plugins/Experimental/NNE/Source/NNERuntimeRDG/Private/NNERuntimeRDGModelHlsl.h
--- a/Engine/Plugins/Experimental/NNE/Source/NNERuntimeRDG/Private/NNERuntimeRDGModelHlsl.h
+++ b/Engine/Plugins/Experimental/NNE/Source/NNERuntimeRDG/Private/NNERuntimeRDGModelHlsl.h
@@ -4,7 +4,6 @@
 
 #include "NNERuntimeRDGModel.h"
 #include "NNECoreTypes.h"
-#include "NNERuntimeRDGTensorHlsl.h"
 
 class FRDGBuilder;
 
@@ -21,25 +20,17 @@ public:
 
 	bool Init(TConstArrayView<uint8> ModelData);
 
-	virtual int SetInputTensorShapes(TConstArrayView<NNECore::FTensorShape> InputShapes) override;
-	virtual int EnqueueRDG(FRDGBuilder& RDGBuilder, TConstArrayView<NNECore::FTensorBindingRDG> InInputBindings, TConstArrayView<NNECore::FTensorBindingRDG> InOutputBindings) override;
-
 protected:
 
-	virtual void AddDispatchOps_RenderThread(FRDGBuilder& GraphBuilder) override { };
 	virtual int PrepareTensorShapesAndData() override;
+	virtual bool AddWeightsToRDGGraph(FRDGBuilder& RDGBuilder) override;
+	virtual void AddDispatchOps_RenderThread(FRDGBuilder& GraphBuilder) override;
 
 	bool PrepareWeights();
 
 private:
 
 	TArray<FOperatorHlsl*>	Operators;
-
-	TArray<FTensorHLSLRef> AllTensorHLSLRefs;
-	TArray<FTensorHLSL> InputTensorHLSLs;
-	TArray<FTensorHLSL> OutputTensorHLSLs;
-	TArray<FTensorHLSL> IntermediateTensorHLSLs;
-	TArray<FTensorHLSL> WeightTensorHLSLs;
 	TArray<TRefCountPtr<FRDGPooledBuffer>> WeightsExternalRDGResources;
 };
 
diff --git a/Engine/Plugins/Experimental/NNE/Source/NNERuntimeRDG/Private/NNERuntimeRDGTensorHlsl.cpp b/Engine/Plugins/Experimental/NNE/Source/NNERuntimeRDG/Private/NNERuntimeRDGTensorHlsl.cpp
deleted file mode 100644
diff --git a/Engine/Plugins/Experimental/NNE/Source/NNERuntimeRDG/Private/NNERuntimeRDGTensorHlsl.h b/Engine/Plugins/Experimental/NNE/Source/NNERuntimeRDG/Private/NNERuntimeRDGTensorHlsl.h
deleted file mode 100644
diff --git a/Engine/Plugins/Experimental/NNE/Source/NNEUtils/Private/NNEUtilsModelBuilderONNX.cpp b/Engine/Plugins/Experimental/NNE/Source/NNEUtils/Private/NNEUtilsModelBuilderONNX.cpp
--- a/Engine/Plugins/Experimental/NNE/Source/NNEUtils/Private/NNEUtilsModelBuilderONNX.cpp
+++ b/Engine/Plugins/Experimental/NNE/Source/NNEUtils/Private/NNEUtilsModelBuilderONNX.cpp
@@ -310,7 +310,7 @@ private:
 //
 void BuildShapeForModel(bool ConvertToVariadicShape, const NNECore::FTensorShape& InShape, TArray<int32>& OutShape)
 {
-	OutShape.Empty();
+	OutShape.Reset();
 	for (int32 Idx = 0; Idx < InShape.Rank(); ++Idx)
 	{
 		int32 Dim = static_cast<int32>(InShape.GetData()[Idx]);
