Use TempBuff binding

#rnx

diff --git a/Engine/Plugins/Experimental/NNE/Source/NNERuntimeRDG/Private/NNERuntimeRDGDml.cpp b/Engine/Plugins/Experimental/NNE/Source/NNERuntimeRDG/Private/NNERuntimeRDGDml.cpp
--- a/Engine/Plugins/Experimental/NNE/Source/NNERuntimeRDG/Private/NNERuntimeRDGDml.cpp
+++ b/Engine/Plugins/Experimental/NNE/Source/NNERuntimeRDG/Private/NNERuntimeRDGDml.cpp
@@ -1010,7 +1010,7 @@ public:
 		if (TempBuff)
 		{
 #ifdef NNE_USE_D3D12_RESOURCES
-			TempBind = { PersistBuff, 0, PersistBuff->GetDesc().Width };
+			TempBind = { TempBuff, 0, TempBuff->GetDesc().Width };
 #else
 			TempBind = MakeBind(TempBuff);
 #endif
