NNE - RDG HLSL runtime - Allow unary operator to forward tensor constness in graph + add Shape operator

Problem:
* Constant tensors can be followed by unary operator before reaching a node needing constant inputs.
* Shape operator was not implemented

Solutions:
* Implement scheduling side tensor computation to propagate constness before the model is runned.
* Add shape operator (no HLSL impl atm)

Tests: all nne test green on editor win64 + style transfer

diff --git a/Engine/Plugins/Experimental/NNE/Source/NNECore/Internal/NNECoreOperator.h b/Engine/Plugins/Experimental/NNE/Source/NNECore/Internal/NNECoreOperator.h
--- a/Engine/Plugins/Experimental/NNE/Source/NNECore/Internal/NNECoreOperator.h
+++ b/Engine/Plugins/Experimental/NNE/Source/NNECore/Internal/NNECoreOperator.h
@@ -2,6 +2,8 @@
 
 #pragma once
 
+#include "HAL/Platform.h"
+
 namespace UE::NNECore::Internal
 {
 
diff --git a/Engine/Plugins/Experimental/NNE/Source/NNECore/Internal/NNECoreTensor.h b/Engine/Plugins/Experimental/NNE/Source/NNECore/Internal/NNECoreTensor.h
--- a/Engine/Plugins/Experimental/NNE/Source/NNECore/Internal/NNECoreTensor.h
+++ b/Engine/Plugins/Experimental/NNE/Source/NNECore/Internal/NNECoreTensor.h
@@ -59,9 +59,12 @@ namespace UE::NNECore::Internal
 			DataSize = static_cast<uint64>(GetTensorDataTypeSizeInBytes(DataType) * Volume);
 		}
 
-		void SetPreparedData(TConstArrayView<uint8> Data)
+		template <typename T> void SetPreparedData(TConstArrayView<T> Data)
 		{
-			checkf(Data.Num() == DataSize, TEXT("Incorrect data size, it should match tensor shape and data type."));
+			const uint8* DataPtr = reinterpret_cast<const uint8*>(Data.GetData());
+			TConstArrayView<uint8> DataAsByte = MakeArrayView(DataPtr, Data.Num() / sizeof(T));
+			
+			checkf(DataAsByte.Num() == DataSize, TEXT("Incorrect data size, it should match tensor shape and data type."));
 			PreparedData.Reset();
 			PreparedData.Append(Data);
 		}
diff --git a/Engine/Plugins/Experimental/NNE/Source/NNERuntimeRDG/Private/Hlsl/NNERuntimeRDGElementWiseUnary.cpp b/Engine/Plugins/Experimental/NNE/Source/NNERuntimeRDG/Private/Hlsl/NNERuntimeRDGElementWiseUnary.cpp
--- a/Engine/Plugins/Experimental/NNE/Source/NNERuntimeRDG/Private/Hlsl/NNERuntimeRDGElementWiseUnary.cpp
+++ b/Engine/Plugins/Experimental/NNE/Source/NNERuntimeRDG/Private/Hlsl/NNERuntimeRDGElementWiseUnary.cpp
@@ -1,6 +1,7 @@
 // Copyright Epic Games, Inc. All Rights Reserved.
 
 #include "NNERuntimeRDGElementWiseUnary.h"
+#include "NNERuntimeRDGElementWiseUnaryHelper.h"
 #include "NNEHlslShadersElementWiseUnaryCS.h"
 #include "NNERuntimeRDGHlslHelper.h"
 #include "NNECoreAttributeMap.h"
@@ -38,6 +39,20 @@ namespace UE::NNERuntimeRDG::Private::Hlsl
 			check(InputTensors.Num() == 1);
 			check(OutputTensors.Num() == 1);
 			OutputTensors[0]->SetShape(InputTensors[0]->GetShape());
+
+			const NNECore::Internal::FTensor& X = *InputTensors[0];
+			if (X.HasPreparedData())
+			{
+				TConstArrayView<float> XData = X.GetPreparedData<float>();
+				TArray<float> OutputData;
+				OutputData.Reserve(XData.Num());
+				for (float elem : XData)
+				{
+					OutputData.Add(ElementWiseUnaryCPUHelper::Apply<OpType>(elem, Alpha, Beta, Gamma));
+				}
+				OutputTensors[0]->SetPreparedData<float>(OutputData);
+			}
+			
 			return 0;
 		}
 
diff --git a/Engine/Plugins/Experimental/NNE/Source/NNERuntimeRDG/Private/Hlsl/NNERuntimeRDGShape.cpp b/Engine/Plugins/Experimental/NNE/Source/NNERuntimeRDG/Private/Hlsl/NNERuntimeRDGShape.cpp
new file mode 100644
--- /dev/null
+++ b/Engine/Plugins/Experimental/NNE/Source/NNERuntimeRDG/Private/Hlsl/NNERuntimeRDGShape.cpp
@@ -0,0 +1,95 @@
+// Copyright Epic Games, Inc. All Rights Reserved.
+
+#include "NNERuntimeRDGShape.h"
+#include "NNECoreTensor.h"
+#include "NNECoreTypes.h"
+
+namespace UE::NNERuntimeRDG::Private::Hlsl
+{
+	/**
+	 * Shape operator implementation
+	 */
+	class FShape : public FOperatorHlsl
+	{
+	public:
+
+		FShape() {}
+		virtual ~FShape() = default;
+
+	public:
+
+		virtual int PrepareOutputs(TConstArrayView<NNECore::Internal::FTensorRef> InputTensors, TArrayView<NNECore::Internal::FTensorRef> OutputTensors) const override
+		{
+			check(InputTensors.Num() == 1);
+			check(OutputTensors.Num() == 1);
+			check(OutputTensors[0]->GetDataType() == ENNETensorDataType::Int64);
+
+			const NNECore::Internal::FTensor& X = *InputTensors[0];
+
+			TArray<uint32> OutputShapeData;
+			
+			OutputShapeData.Emplace(X.GetShape().Rank());
+			
+			NNECore::FTensorShape OutputShape = NNECore::FTensorShape::Make(OutputShapeData);
+			TArray<int64> OutputData;
+			
+			OutputTensors[0]->SetShape(OutputShape);
+			for (int32 i = 0; i < X.GetShape().Rank(); ++i)
+			{
+				OutputData.Emplace(X.GetShape().GetData()[i]);
+			}
+			OutputTensors[0]->SetPreparedData<int64>(OutputData);
+			
+			return 0;
+		};
+
+		virtual bool Initialize(TConstArrayView<NNECore::FTensorDesc> InputTensorDescs, TConstArrayView<NNECore::FTensorDesc> OutputTensorDescs, const NNECore::FAttributeMap& Attributes) override
+		{
+			check(InputTensorDescs.Num() == 1);
+			check(OutputTensorDescs.Num() == 1);
+
+			if (OutputTensorDescs[0].GetDataType() != ENNETensorDataType::Int64)
+			{
+				UE_LOG(LogNNE, Warning, TEXT("Shape should output a tensor of type Int64"));
+				return false;
+			}
+			
+			return true;
+		}
+
+		virtual void Dispatch(FRDGBuilder& GraphBuilder, TConstArrayView<FTensorRDGRef> InputTensors, TConstArrayView<FTensorRDGRef> OutputTensors) override
+		{
+			//Need to upload the data to gpu tensor or have the hlsl model understand that this operator result will only be needed on CPU.
+			//Current behavior will be faulty if the model rely on this operator result on GPU.
+		}
+	};
+
+	bool ValidateShapeOperator(const NNECore::FAttributeMap& AttributeMap, TConstArrayView<ENNETensorDataType> InputTypes, TConstArrayView<NNECore::FSymbolicTensorShape> InputShapes)
+	{
+		bool bIsValid = true;
+
+		//This match version 1 of the shape operator (next version is with opset 13)
+		//[REDACTED_PROFILE_URL]/onnx/blob/main/docs/Changelog.md#Shape-1
+		FAttributeValidator AttributeValidator;
+		bIsValid &= AttributeValidator.Validate(AttributeMap);
+
+		FInputValidator InputValidator;
+		InputValidator.AddSupportedType(ENNETensorDataType::Float);
+		InputValidator.AddSupportedType(ENNETensorDataType::Int64);
+		InputValidator.AddRequired();
+		bIsValid &= InputValidator.Validate(InputTypes);
+
+		return bIsValid;
+	}
+
+	FOperatorHlsl* CreateShapeOperator()
+	{
+		return new FShape();
+	}
+
+	bool RegisterShapeOperator(FOperatorRegistryHlsl& Registry)
+	{
+		Registry.OpAdd(TEXT("Shaoe"), CreateShapeOperator, ValidateShapeOperator);
+		return true;
+	}
+} // UE::NNERuntimeRDG::Private::Hlsl
diff --git a/Engine/Plugins/Experimental/NNE/Source/NNERuntimeRDG/Private/Hlsl/NNERuntimeRDGShape.h b/Engine/Plugins/Experimental/NNE/Source/NNERuntimeRDG/Private/Hlsl/NNERuntimeRDGShape.h
new file mode 100644
--- /dev/null
+++ b/Engine/Plugins/Experimental/NNE/Source/NNERuntimeRDG/Private/Hlsl/NNERuntimeRDGShape.h
@@ -0,0 +1,11 @@
+// Copyright Epic Games, Inc. All Rights Reserved.
+
+#pragma once
+
+#include "NNERuntimeRDG.h"
+#include "NNERuntimeRDGHlslOp.h"
+
+namespace UE::NNERuntimeRDG::Private::Hlsl
+{
+	bool RegisterShapeOperator(FOperatorRegistryHlsl& Registry);
+} // UE::NNERuntimeRDG::Private::Hlsl
diff --git a/Engine/Plugins/Experimental/NNE/Source/NNERuntimeRDG/Private/NNERuntimeRDGElementWiseUnaryHelper.cpp b/Engine/Plugins/Experimental/NNE/Source/NNERuntimeRDG/Private/NNERuntimeRDGElementWiseUnaryHelper.cpp
new file mode 100644
--- /dev/null
+++ b/Engine/Plugins/Experimental/NNE/Source/NNERuntimeRDG/Private/NNERuntimeRDGElementWiseUnaryHelper.cpp
@@ -0,0 +1,115 @@
+// Copyright Epic Games, Inc. All Rights Reserved.
+
+#include "NNERuntimeRDGElementWiseUnaryHelper.h"
+#include "Math/UnrealMathUtility.h"
+
+namespace UE::NNERuntimeRDG::Private::ElementWiseUnaryCPUHelper
+{
+	template<> float Apply<NNECore::Internal::EElementWiseUnaryOperatorType::Abs>(float X, float Alpha, float Beta, float Gamma) { return FMath::Abs(X); }
+
+	template<> float Apply<NNECore::Internal::EElementWiseUnaryOperatorType::Acos>(float X, float Alpha, float Beta, float Gamma) { return FMath::Acos(X); }
+	
+	template<> float Apply<NNECore::Internal::EElementWiseUnaryOperatorType::Acosh>(float X, float Alpha, float Beta, float Gamma) {
+		//https://mathworld.wolfram.com/InverseHyperbolicCosine.html
+		float FloatNan = FMath::Sqrt(-1.0f);
+		float yAboveOne = FMath::Loge(X + FMath::Sqrt(X + 1.0f) + FMath::Sqrt(X - 1.0f));
+		return (X == 1.0f) ? 0.0f : (X >= 1.0f) ? yAboveOne : FloatNan;
+	}
+	
+	template<> float Apply<NNECore::Internal::EElementWiseUnaryOperatorType::Asin>(float X, float Alpha, float Beta, float Gamma) { return FMath::Asin(X); }
+	
+	template<> float Apply<NNECore::Internal::EElementWiseUnaryOperatorType::Asinh>(float X, float Alpha, float Beta, float Gamma) {
+		//https://mathworld.wolfram.com/InverseHyperbolicSine.html
+		return FMath::Loge(X + FMath::Sqrt(1 + (X * X)));
+	}
+	
+	template<> float Apply<NNECore::Internal::EElementWiseUnaryOperatorType::Atan>(float X, float Alpha, float Beta, float Gamma) { return FMath::Atan(X); }
+	
+	template<> float Apply<NNECore::Internal::EElementWiseUnaryOperatorType::Atanh>(float X, float Alpha, float Beta, float Gamma) {
+		//https://mathworld.wolfram.com/InverseHyperbolicTangent.html
+		return 0.5f * (FMath::Loge(1 + X) - FMath::Loge(1 - X));
+	}
+
+	template<> float Apply<NNECore::Internal::EElementWiseUnaryOperatorType::Ceil>(float X, float Alpha, float Beta, float Gamma) { return FMath::CeilToFloat(X); }
+	
+	template<> float Apply<NNECore::Internal::EElementWiseUnaryOperatorType::Cos>(float X, float Alpha, float Beta, float Gamma) { return FMath::Cos(X); }
+	
+	template<> float Apply<NNECore::Internal::EElementWiseUnaryOperatorType::Cosh>(float X, float Alpha, float Beta, float Gamma) {
+		//https://mathworld.wolfram.com/HyperbolicCosine.html
+		return 0.5f * (FMath::Exp(X) - FMath::Exp(-X));
+	}
+	
+	template<> float Apply<NNECore::Internal::EElementWiseUnaryOperatorType::Elu>(float X, float Alpha, float Beta, float Gamma) {
+		//[REDACTED_PROFILE_URL]/onnx/blob/main/docs/Operators.md#elu
+		float yNeg = Alpha * (FMath::Exp(X) - 1.0f);
+		float yPosOrZero = X;
+		return (X >= 0.0f) ? yPosOrZero : yNeg;
+	}
+	
+	template<> float Apply<NNECore::Internal::EElementWiseUnaryOperatorType::Exp>(float X, float Alpha, float Beta, float Gamma) { return FMath::Exp(X); }
+
+	template<> float Apply<NNECore::Internal::EElementWiseUnaryOperatorType::Floor>(float X, float Alpha, float Beta, float Gamma) { return FMath::Floor(X); }
+
+	template<> float Apply<NNECore::Internal::EElementWiseUnaryOperatorType::IsInf>(float X, float Alpha, float Beta, float Gamma) { return !FMath::IsFinite(X); }
+
+	template<> float Apply<NNECore::Internal::EElementWiseUnaryOperatorType::IsNan>(float X, float Alpha, float Beta, float Gamma) { return FMath::IsNaN(X); }
+
+	template<> float Apply<NNECore::Internal::EElementWiseUnaryOperatorType::HardSigmoid>(float X, float Alpha, float Beta, float Gamma) {
+		//[REDACTED_PROFILE_URL]/onnx/blob/main/docs/Operators.md#hardSigmoid
+		return FMath::Max(0.0f, FMath::Min(1.0f, Alpha * X + Beta));
+	}
+
+	template<> float Apply<NNECore::Internal::EElementWiseUnaryOperatorType::HardSwish>(float X, float Alpha, float Beta, float Gamma) {
+		//[REDACTED_PROFILE_URL]/onnx/blob/main/docs/Operators.md#hardSwish
+		return Apply<NNECore::Internal::EElementWiseUnaryOperatorType::HardSigmoid>(X, 1.0f / 6.0f, 0.5f, Gamma);
+	}
+
+	template<> float Apply<NNECore::Internal::EElementWiseUnaryOperatorType::LeakyRelu>(float X, float Alpha, float Beta, float Gamma) { return (X >= 0.0f) ? X : Alpha * X; }
+
+	template<> float Apply<NNECore::Internal::EElementWiseUnaryOperatorType::Log>(float X, float Alpha, float Beta, float Gamma) { return FMath::Loge(X); }
+
+	template<> float Apply<NNECore::Internal::EElementWiseUnaryOperatorType::Neg>(float X, float Alpha, float Beta, float Gamma) { return -X; }
+
+	template<> float Apply<NNECore::Internal::EElementWiseUnaryOperatorType::Reciprocal>(float X, float Alpha, float Beta, float Gamma) { return 1.0f/X; }
+
+	template<> float Apply<NNECore::Internal::EElementWiseUnaryOperatorType::Relu>(float X, float Alpha, float Beta, float Gamma) { return FMath::Max(X, 0.0f); }
+
+	template<> float Apply<NNECore::Internal::EElementWiseUnaryOperatorType::Round>(float X, float Alpha, float Beta, float Gamma) { return FMath::RoundToFloat(X); }
+
+	template<> float Apply<NNECore::Internal::EElementWiseUnaryOperatorType::Selu>(float X, float Alpha, float Beta, float Gamma) {
+		//[REDACTED_PROFILE_URL]/onnx/blob/main/docs/Operators.md#Selu
+		float yNegOrZero = Gamma * (Alpha * FMath::Exp(X) - Alpha);
+		float yPos = Gamma * X;
+		return (X > 0.0f) ? yPos : yNegOrZero;
+	}
+	
+	template<> float Apply<NNECore::Internal::EElementWiseUnaryOperatorType::Sigmoid>(float X, float Alpha, float Beta, float Gamma) { return 1.0f / (1.0f + FMath::Exp(-X)); }
+	
+	template<> float Apply<NNECore::Internal::EElementWiseUnaryOperatorType::Sign>(float X, float Alpha, float Beta, float Gamma) { return FMath::Sign(X); }
+	
+	template<> float Apply<NNECore::Internal::EElementWiseUnaryOperatorType::Sin>(float X, float Alpha, float Beta, float Gamma) { return FMath::Sin(X); }
+	
+	template<> float Apply<NNECore::Internal::EElementWiseUnaryOperatorType::Sinh>(float X, float Alpha, float Beta, float Gamma) { return FMath::Sinh(X); }
+	
+	template<> float Apply<NNECore::Internal::EElementWiseUnaryOperatorType::Softplus>(float X, float Alpha, float Beta, float Gamma) { return FMath::Loge(FMath::Exp(X) + 1.0f); }
+	
+	template<> float Apply<NNECore::Internal::EElementWiseUnaryOperatorType::Softsign>(float X, float Alpha, float Beta, float Gamma) { return X / (1.0f + FMath::Abs(X)); }
+	
+	template<> float Apply<NNECore::Internal::EElementWiseUnaryOperatorType::Sqrt>(float X, float Alpha, float Beta, float Gamma) { return FMath::Sqrt(X); }
+	
+	template<> float Apply<NNECore::Internal::EElementWiseUnaryOperatorType::Tan>(float X, float Alpha, float Beta, float Gamma) { return FMath::Tan(X); }
+
+	template<> float Apply<NNECore::Internal::EElementWiseUnaryOperatorType::Tanh>(float X, float Alpha, float Beta, float Gamma) {
+		//https://mathworld.wolfram.com/HyperbolicTangent.html
+		return Apply<NNECore::Internal::EElementWiseUnaryOperatorType::Sinh>(X, Alpha, Beta, Gamma) / Apply<NNECore::Internal::EElementWiseUnaryOperatorType::Cosh>(X, Alpha, Beta, Gamma);
+	}
+
+	template<> float Apply<NNECore::Internal::EElementWiseUnaryOperatorType::Erf>(float X, float Alpha, float Beta, float Gamma) {
+		//https://aapt.scitation.org/doi/abs/10.1119/1.15018?journalCode=ajp
+		float a = 167.0f / 148.0f;
+		float b = 11.0f / 109.0f;
+		float x3 = X * X * X;
+		return Apply<NNECore::Internal::EElementWiseUnaryOperatorType::Tanh>(a * X + b * x3, Alpha, Beta, Gamma);
+	}
+	
+} // UE::NNERuntimeRDG::Private::ElementWiseUnaryCPUHelper
diff --git a/Engine/Plugins/Experimental/NNE/Source/NNERuntimeRDG/Private/NNERuntimeRDGElementWiseUnaryHelper.h b/Engine/Plugins/Experimental/NNE/Source/NNERuntimeRDG/Private/NNERuntimeRDGElementWiseUnaryHelper.h
new file mode 100644
--- /dev/null
+++ b/Engine/Plugins/Experimental/NNE/Source/NNERuntimeRDG/Private/NNERuntimeRDGElementWiseUnaryHelper.h
@@ -0,0 +1,11 @@
+// Copyright Epic Games, Inc. All Rights Reserved.
+
+#pragma once 
+
+#include "NNECoreOperator.h"
+
+namespace UE::NNERuntimeRDG::Private::ElementWiseUnaryCPUHelper
+{
+	template<NNECore::Internal::EElementWiseUnaryOperatorType OpType> float Apply(float X, float Alpha, float Beta, float Gamma);
+	
+} // UE::NNERuntimeRDG::Private::ElementWiseUnaryCPUHelper
diff --git a/Engine/Plugins/Experimental/NNE/Source/NNERuntimeRDG/Private/NNERuntimeRDGHlsl.cpp b/Engine/Plugins/Experimental/NNE/Source/NNERuntimeRDG/Private/NNERuntimeRDGHlsl.cpp
--- a/Engine/Plugins/Experimental/NNE/Source/NNERuntimeRDG/Private/NNERuntimeRDGHlsl.cpp
+++ b/Engine/Plugins/Experimental/NNE/Source/NNERuntimeRDG/Private/NNERuntimeRDGHlsl.cpp
@@ -17,6 +17,7 @@
 #include "Hlsl/NNERuntimeRDGGemm.h"
 #include "Hlsl/NNERuntimeRDGInstanceNormalization.h"
 #include "Hlsl/NNERuntimeRDGPad.h"
+#include "Hlsl/NNERuntimeRDGShape.h"
 #include "Hlsl/NNERuntimeRDGUpsample.h"
 #include "Hlsl/NNERuntimeRDGMatMul.h"
 
@@ -40,6 +41,7 @@ bool UNNERuntimeRDGHlslImpl::Init()
 	RegisterInstanceNormalizationOperator(*registry);
 	RegisterUpsampleOperator(*registry);
 	RegisterPadOperator(*registry);
+	RegisterShapeOperator(*registry);
 
 	return true;
 }
