Reorganize RDG DirectML into multple files

#rnx

diff --git a/Engine/Plugins/Experimental/NNE/Source/NNERuntimeRDG/Private/Dml/NNEDmlCommon.h b/Engine/Plugins/Experimental/NNE/Source/NNERuntimeRDG/Private/Dml/NNEDmlCommon.h
new file mode 100644
--- /dev/null
+++ b/Engine/Plugins/Experimental/NNE/Source/NNERuntimeRDG/Private/Dml/NNEDmlCommon.h
@@ -0,0 +1,24 @@
+// Copyright Epic Games, Inc. All Rights Reserved.
+
+#pragma once
+
+#ifdef NNE_USE_DIRECTML
+
+#include "HAL/Platform.h"
+
+#include "Windows/AllowWindowsPlatformTypes.h"
+THIRD_PARTY_INCLUDES_START
+#include <unknwn.h>
+#include "Microsoft/COMPointer.h"
+#include "DirectML.h"
+THIRD_PARTY_INCLUDES_END
+#include "Windows/HideWindowsPlatformTypes.h"
+
+// DirectML is implemented using COM on all platforms
+#ifdef IID_GRAPHICS_PPV_ARGS
+#define DML_PPV_ARGS(x) __uuidof(*x), IID_PPV_ARGS_Helper(x)
+#else
+#define DML_PPV_ARGS(x) IID_PPV_ARGS(x)
+#endif
+
+#endif // NNE_USE_DIRECTML
diff --git a/Engine/Plugins/Experimental/NNE/Source/NNERuntimeRDG/Private/Dml/NNEDmlModel.cpp b/Engine/Plugins/Experimental/NNE/Source/NNERuntimeRDG/Private/Dml/NNEDmlModel.cpp
new file mode 100644
--- /dev/null
+++ b/Engine/Plugins/Experimental/NNE/Source/NNERuntimeRDG/Private/Dml/NNEDmlModel.cpp
@@ -0,0 +1,1235 @@
+// Copyright Epic Games, Inc. All Rights Reserved.
+
+#include "NNEDmlModel.h"
+
+#ifdef NNE_USE_DIRECTML
+#include "NNEDmlOperator.h"
+
+#include "RenderGraphBuilder.h"
+#include "RenderGraphUtils.h"
+
+#include "ID3D12DynamicRHI.h"
+#endif
+
+namespace UE::NNERuntimeRDG::Private::Dml
+{
+
+FModelInfo* FModelInfo::Get()
+{
+	static FModelInfo Inst;
+
+	return &Inst;
+}
+
+FModelInfo::FModelInfo()
+	: Guid((int32)'R', (int32)'D', (int32)'G', (int32)'D')
+	, Version(0x1)
+{
+}
+
+FGuid FModelInfo::GetGuid()
+{
+	return Guid;
+}
+
+int32 FModelInfo::GetVersion()
+{
+	return Version;
+}
+
+int32 FModelInfo::GetGuidSize()
+{
+	return sizeof(Guid);
+}
+
+int32 FModelInfo::GetVersionSize()
+{
+	return sizeof(Version);
+}
+
+bool FModelInfo::ValidateGuidAndVersion(const uint8* InGuid, const uint8* InVersion)
+{
+	bool bResult;
+	
+	bResult = FGenericPlatformMemory::Memcmp(InGuid, &Guid, GetGuidSize()) == 0;
+	bResult &= FGenericPlatformMemory::Memcmp(InVersion, &Version, GetVersionSize()) == 0;
+
+	return bResult;
+}
+
+#ifdef NNE_USE_DIRECTML
+
+//
+//
+//
+class FModel::FDebugName
+{
+	static constexpr int32 Size = 128;
+
+public:
+
+	FDebugName()
+	{
+		Str[0] = '\0';
+		Length = 0;
+	}
+
+	FDebugName(const FString& InStr)
+	{
+		FTCHARToUTF8 Conv(*InStr);
+
+		Length = Conv.Length() + 1 < Size ? Conv.Length() + 1 : Size - 1;
+		FCStringAnsi::Strncpy(Str, Conv.Get(), Length);
+		Str[Length] = '\0';
+	}
+
+	FDebugName(FStringView InStr)
+	{
+		FTCHARToUTF8 Conv(InStr.GetData());
+
+		Length = Conv.Length() + 1 < Size ? Conv.Length() + 1 : Size - 1;
+		FCStringAnsi::Strncpy(Str, Conv.Get(), Length);
+		Str[Length] = '\0';
+	}
+
+	~FDebugName()
+	{
+	}
+
+	const char* Get() const
+	{
+		return Str;
+	}
+
+private:
+
+	char	Str[Size];
+	int32	Length;
+};
+
+//
+//
+//
+class FModel::FBindingTable
+{
+public:
+
+	bool Init(FModel* InModel)
+	{
+		Model = InModel;
+		DynamicRHI = InModel->DynamicRHI;
+
+		return true;
+	}
+
+#ifdef NNE_USE_D3D12_RESOURCES
+	void Bind(IDMLOperatorInitializer* OpInit, TConstArrayView<FRHIBuffer*> InputBuffers, ID3D12Resource* PersistBuff, ID3D12Resource* TempBuff = nullptr)
+#else
+	void Bind(IDMLOperatorInitializer* OpInit, TConstArrayView<FRHIBuffer*> InputBuffers, FRHIBuffer* PersistBuff, FRHIBuffer* TempBuff = nullptr)
+#endif
+	{
+		Reset(OpInit);
+
+		TArray<DML_BUFFER_BINDING, TInlineAllocator<MaxNumInputs>> Inputs;
+
+		for (FRHIBuffer* Buffer : InputBuffers)
+		{
+			if (Buffer)
+			{
+				Inputs.Emplace(MakeBind(Buffer));
+			}
+			else
+			{
+				Inputs.Add({});
+			}
+		}
+
+		DML_BUFFER_ARRAY_BINDING	InputBindArray{ Inputs.Num(), Inputs.GetData() };
+		DML_BINDING_DESC			InputBindArrayDesc{ DML_BINDING_TYPE_BUFFER_ARRAY, &InputBindArray };
+		
+		BindingTable->BindInputs(1, &InputBindArrayDesc);
+		
+		DML_BUFFER_BINDING			PersistBind {};
+		DML_BINDING_DESC			PersistBindDesc{ DML_BINDING_TYPE_BUFFER, &PersistBind };
+
+		if (PersistBuff)
+		{
+#ifdef NNE_USE_D3D12_RESOURCES
+			PersistBind = DML_BUFFER_BINDING { PersistBuff, 0, PersistBuff->GetDesc().Width };
+#else
+			PersistBind = MakeBind(PersistBuff);
+#endif
+		}
+
+		BindingTable->BindOutputs(1, &PersistBindDesc);
+
+		DML_BUFFER_BINDING			TempBind{};
+		DML_BINDING_DESC			TempBindDesc{ DML_BINDING_TYPE_BUFFER, &TempBind };
+
+		if (TempBuff)
+		{
+#ifdef NNE_USE_D3D12_RESOURCES
+			TempBind = { TempBuff, 0, TempBuff->GetDesc().Width };
+#else
+			TempBind = MakeBind(TempBuff);
+#endif
+			BindingTable->BindTemporaryResource(&TempBindDesc);
+		}
+	}
+
+#ifdef NNE_USE_D3D12_RESOURCES
+	void Bind(IDMLCompiledOperator* Op, TConstArrayView<FRHIBuffer*> InputBuffers, TConstArrayView<FRHIBuffer*> OutputBuffers, ID3D12Resource* PersistBuff = nullptr, ID3D12Resource* TempBuff = nullptr)
+#else
+	
+	void Bind(IDMLCompiledOperator* Op, TConstArrayView<FRHIBuffer*> InputBuffers, TConstArrayView<FRHIBuffer*> OutputBuffers, FRHIBuffer* PersistBuff = nullptr, FRHIBuffer* TempBuff = nullptr)
+#endif
+	{
+		Reset(Op);
+
+		for (FRHIBuffer* Buffer : InputBuffers)
+		{
+			AddBind(Buffer, InputBinds, InputBindDescs);
+		}
+
+		for (FRHIBuffer* Buffer : OutputBuffers)
+		{
+			AddBind(Buffer, OutputBinds, OutputBindDescs);
+		}
+
+		BindingTable->BindInputs(InputBinds.Num(), InputBindDescs.GetData());
+		BindingTable->BindOutputs(OutputBinds.Num(), OutputBindDescs.GetData());
+
+		DML_BUFFER_BINDING			PersistBind;
+		DML_BINDING_DESC			PersistBindDesc{ DML_BINDING_TYPE_BUFFER, &PersistBind };
+
+		if (PersistBuff)
+		{
+#ifdef NNE_USE_D3D12_RESOURCES
+			PersistBind = { PersistBuff, 0, PersistBuff->GetDesc().Width };
+#else
+			PersistBind =  MakeBind(PersistBuff);MakeBind(PersistBuff);
+#endif
+			BindingTable->BindPersistentResource(&PersistBindDesc);
+		}
+
+		DML_BUFFER_BINDING			TempBind{};
+		DML_BINDING_DESC			TempBindDesc{ DML_BINDING_TYPE_BUFFER, &TempBind };
+
+		if (TempBuff)
+		{
+			
+#ifdef NNE_USE_D3D12_RESOURCES
+			TempBind = { TempBuff, 0, TempBuff->GetDesc().Width };
+#else
+			TempBind = MakeBind(TempBuff);
+#endif
+			BindingTable->BindTemporaryResource(&TempBindDesc);
+		}
+	}
+
+
+	IDMLBindingTable* Get()
+	{
+		return BindingTable;
+	}
+
+private:
+
+	bool Reset(IDMLDispatchable* Disp)
+	{
+		InputBinds.Reset();
+		InputBindDescs.Reset();
+		OutputBinds.Reset();
+		OutputBindDescs.Reset();
+		
+		DML_BINDING_PROPERTIES BindingProps = Disp->GetBindingProperties();
+		DML_BINDING_TABLE_DESC Desc{};
+
+		Desc.Dispatchable = Disp;
+		Desc.CPUDescriptorHandle = CD3DX12_CPU_DESCRIPTOR_HANDLE(Model->DescHeap->GetCPUDescriptorHandleForHeapStart(), 0, Model->DescSize);
+		Desc.GPUDescriptorHandle = CD3DX12_GPU_DESCRIPTOR_HANDLE(Model->DescHeap->GetGPUDescriptorHandleForHeapStart(), 0, Model->DescSize);
+		Desc.SizeInDescriptors = Model->DescCount;
+
+		HRESULT Res;
+
+		if (!BindingTable)
+		{
+			Res = Model->DevCtx->Device->CreateBindingTable(&Desc, DML_PPV_ARGS(&BindingTable));
+			if (!BindingTable)
+			{
+				UE_LOG(LogNNE, Warning, TEXT("Failed to create DML binding table, res:%d"), Res);
+				return false;
+			}
+		}
+		else
+		{
+			BindingTable->Reset(&Desc);
+		}
+
+		return true;
+	}
+
+	template<class TBindingArray, class TDescArray>
+	void AddBind(FRHIBuffer* Buffer, TBindingArray& Bindings, TDescArray& Descs)
+	{
+		if (Buffer)
+		{
+			DML_BUFFER_BINDING& Bind = Bindings.Add_GetRef(MakeBind(Buffer));
+			Descs.Add({ DML_BINDING_TYPE_BUFFER, &Bind });
+		}
+		else
+		{
+			DML_BUFFER_BINDING& Bind = Bindings.Add_GetRef({});
+			Descs.Add({ DML_BINDING_TYPE_BUFFER, &Bind });
+		}
+	}
+
+	DML_BUFFER_BINDING MakeBind(FRHIBuffer* Buffer)
+	{
+		ID3D12Resource* Resource = DynamicRHI->RHIGetResource(Buffer);
+
+		return DML_BUFFER_BINDING{ Resource, 0, Buffer->GetSize() };
+	}
+
+	TComPtr<IDMLBindingTable>										BindingTable;
+	TArray<DML_BUFFER_BINDING, TInlineAllocator<MaxNumInputs>>		InputBinds;
+	TArray<DML_BINDING_DESC, TInlineAllocator<MaxNumInputs>>		InputBindDescs;
+	TArray<DML_BUFFER_BINDING, TInlineAllocator<MaxNumOutputs>>		OutputBinds;
+	TArray<DML_BINDING_DESC, TInlineAllocator<MaxNumOutputs>>		OutputBindDescs;
+	ID3D12DynamicRHI*												DynamicRHI;
+	FModel*															Model;
+};
+
+//
+//
+//
+class FModel::FGraphBuilder
+{
+public:
+
+	struct FOpDesc
+	{
+		FOperatorDml*	Op;
+		int32			InputStart;
+		int32			InputCount;
+		int32			OutputStart;
+		int32			OutputCount;
+		FDebugName		DbgName;
+	};
+
+	struct FGraphDesc
+	{
+		TConstArrayView<NNECore::Internal::FTensor>	AllTensors;
+		TConstArrayView<int32>		InputIndices;
+		TConstArrayView<int32>		OutputIndices;
+		TConstArrayView<int32>		WeightIndices;
+		TConstArrayView<int32>		IntermediateIndices;
+		TConstArrayView<FTensorRDG>	WeightTensors;
+		TConstArrayView<FOpDesc>	Operators;
+		TConstArrayView<int32>		OpInputIndices;
+		TConstArrayView<int32>		OpOutputIndices;
+	};
+	
+private:
+
+	enum class EEdgeType
+	{
+		Input,
+		Output,
+		Intermediate
+	};
+
+	struct FEdge
+	{
+		EEdgeType	Type;
+		int32		TensorIdx{ -1 };
+		int32		NodeSrc{ -1 };
+		int32		NodeSrcOutput{ -1 };
+		int32		NodeDst{ -1 };
+		int32		NodeDstInput{ -1 };
+		
+		FEdge(EEdgeType InType)
+			: Type(InType)
+		{
+		}
+
+		FEdge& SetTensorIdx(int32 Value)
+		{
+			TensorIdx = Value;
+			return *this;
+		}
+
+		FEdge& SetNodeSrc(int32 Value)
+		{
+			NodeSrc = Value;
+			return *this;
+		}
+
+		FEdge& SetNodeSrcOutput(int32 Value)
+		{
+			NodeSrcOutput = Value;
+			return *this;
+		}
+
+		FEdge& SetNodeDst(int32 Value)
+		{
+			NodeDst = Value;
+			return *this;
+		}
+
+		FEdge& SetNodeDstInput(int32 Value)
+		{
+			NodeDstInput = Value;
+			return *this;
+		}
+	};
+
+public:
+
+	IDMLCompiledOperator* Compile(FDmlDeviceContext* DevCtx, const FGraphDesc& InGraph)
+	{
+		IDMLDevice*				Device = DevCtx->Device;
+		TComPtr<IDMLDevice1>	Device1;
+
+		Device1.FromQueryInterface(__uuidof(IDMLDevice1), DevCtx->Device);
+		check(Device1);
+		if (!Device1)
+		{
+			return nullptr;
+		}
+
+		if (!AddEdges(InGraph))
+		{
+			return nullptr;
+		}
+
+		TArray<DML_INPUT_GRAPH_EDGE_DESC>			InputEdges;
+		TArray<DML_OUTPUT_GRAPH_EDGE_DESC>			OutputEdges;
+		TArray<DML_INTERMEDIATE_GRAPH_EDGE_DESC>	IntermediateEdges;
+		TArray<FDebugName>							DbgInputNames;
+		TArray<FDebugName>							DbgIntermediateNames;
+		TArray<FDebugName>							DbgOutputNames;
+
+		DbgInputNames.Reserve(InGraph.AllTensors.Num());
+		DbgIntermediateNames.Reserve(InGraph.AllTensors.Num());
+
+		for (const FEdge& Edge : Edges)
+		{
+			if (Edge.Type == EEdgeType::Input)
+			{
+				check(Edge.NodeSrcOutput >= 0);
+				check(Edge.NodeDst >= 0);
+				check(Edge.NodeDstInput >= 0);
+
+				DML_INPUT_GRAPH_EDGE_DESC& Input = InputEdges.Add_GetRef({});
+
+				Input.GraphInputIndex = Edge.NodeSrcOutput;
+				Input.ToNodeIndex = Edge.NodeDst;
+				Input.ToNodeInputIndex = Edge.NodeDstInput;
+				
+				const FDebugName& DbgName = DbgInputNames.Add_GetRef(InGraph.AllTensors[Edge.TensorIdx].GetName());
+				Input.Name = DbgName.Get();
+			}
+			else if (Edge.Type == EEdgeType::Output)
+			{
+				check(Edge.NodeDstInput >= 0);
+				check(Edge.NodeSrc >= 0);
+				check(Edge.NodeSrcOutput >= 0);
+
+				DML_OUTPUT_GRAPH_EDGE_DESC&	Output = OutputEdges.Add_GetRef({});
+
+				Output.GraphOutputIndex = Edge.NodeDstInput;
+				Output.FromNodeIndex = Edge.NodeSrc;
+				Output.FromNodeOutputIndex = Edge.NodeSrcOutput;
+
+				const FDebugName& DbgName = DbgOutputNames.Add_GetRef(InGraph.AllTensors[Edge.TensorIdx].GetName());
+				Output.Name = DbgName.Get();
+			}
+			else if (Edge.Type == EEdgeType::Intermediate)
+			{
+				DML_INTERMEDIATE_GRAPH_EDGE_DESC& Intermediate = IntermediateEdges.Add_GetRef({});
+
+				check(Edge.NodeSrc >= 0);
+				check(Edge.NodeSrcOutput >= 0);
+				check(Edge.NodeDst >= 0);
+				check(Edge.NodeDstInput >= 0);
+
+				Intermediate.FromNodeIndex = Edge.NodeSrc;
+				Intermediate.FromNodeOutputIndex = Edge.NodeSrcOutput;
+				Intermediate.ToNodeIndex = Edge.NodeDst;
+				Intermediate.ToNodeInputIndex = Edge.NodeDstInput;
+
+				const FDebugName& DbgName = DbgIntermediateNames.Add_GetRef(InGraph.AllTensors[Edge.TensorIdx].GetName());
+				Intermediate.Name = DbgName.Get();
+			}
+		}
+			
+		TArray<DML_GRAPH_NODE_DESC>		Nodes;
+		TArray<DML_GRAPH_EDGE_DESC>		InputEdgeDescs;
+		TArray<DML_GRAPH_EDGE_DESC>		OutputEdgeDescs;
+		TArray<DML_GRAPH_EDGE_DESC>		IntermediateEdgeDescs;
+
+		Nodes.SetNumUninitialized(Operators.Num());
+
+		for (int32 Idx = 0; Idx < Operators.Num(); ++Idx)
+		{
+			DML_GRAPH_NODE_DESC& NodeDesc = Nodes[Idx];
+
+			NodeDesc.Type = DML_GRAPH_NODE_TYPE_OPERATOR;
+			NodeDesc.Desc = &Operators[Idx];
+		}
+
+		for (int32 Idx = 0; Idx < InputEdges.Num(); ++Idx)
+		{
+			DML_GRAPH_EDGE_DESC& Edge = InputEdgeDescs.Add_GetRef({});
+
+			Edge.Type = DML_GRAPH_EDGE_TYPE_INPUT;
+			Edge.Desc = &InputEdges[Idx];
+		}			
+
+		for (int32 Idx = 0; Idx < OutputEdges.Num(); ++Idx)
+		{
+			DML_GRAPH_EDGE_DESC& Edge = OutputEdgeDescs.Add_GetRef({});
+
+			Edge.Type = DML_GRAPH_EDGE_TYPE_OUTPUT;
+			Edge.Desc = &OutputEdges[Idx];				
+		}
+
+		for (int32 Idx = 0; Idx < IntermediateEdges.Num(); ++Idx)
+		{
+			DML_GRAPH_EDGE_DESC& Edge = IntermediateEdgeDescs.Add_GetRef({});
+
+			Edge.Type = DML_GRAPH_EDGE_TYPE_INTERMEDIATE;
+			Edge.Desc = &IntermediateEdges[Idx];
+		}
+			
+		DML_GRAPH_DESC	Graph = DML_GRAPH_DESC{};
+
+		Graph.InputCount = InputEdges.Num();
+		Graph.OutputCount = OutputEdges.Num();
+		Graph.NodeCount = Operators.Num();
+		Graph.Nodes = Nodes.GetData();
+		Graph.InputEdgeCount = InputEdgeDescs.Num();
+		Graph.InputEdges = InputEdgeDescs.GetData();
+		Graph.OutputEdgeCount = OutputEdgeDescs.Num();
+		Graph.OutputEdges = OutputEdgeDescs.GetData();
+		Graph.IntermediateEdgeCount = IntermediateEdgeDescs.Num();
+		Graph.IntermediateEdges = IntermediateEdgeDescs.GetData();
+
+		IDMLCompiledOperator* Op = nullptr;
+		HRESULT Res;
+			
+		Res = Device1->CompileGraph(&Graph, DML_EXECUTION_FLAG_NONE, DML_PPV_ARGS(&Op));
+		if (FAILED(Res))
+		{
+			UE_LOG(LogNNE, Warning, TEXT("Error:Failed to compile DML graph"));
+			Op = nullptr;
+		};
+
+		return Op;
+	}
+
+private:
+
+	bool AddEdges(const FGraphDesc& InGraph)
+	{
+		Edges.Reset();
+		Operators.Reset();
+		NumInputs = 0;
+		NumOutputs = 0;
+
+		for (int32 Idx = 0; Idx < InGraph.InputIndices.Num(); ++Idx)
+		{
+			const int32 TensorIdx = InGraph.InputIndices[Idx];
+			AddInput(TensorIdx);
+		}
+
+		for (int32 Idx = 0; Idx < InGraph.WeightIndices.Num(); ++Idx)
+		{
+			const int32 TensorIdx = InGraph.WeightIndices[Idx];
+			AddInput(TensorIdx);
+		}
+
+		//const int32 OutputIndexStart = InGraph.InputIndices.Num() + InGraph.WeightIndices.Num() + InGraph.IntermediateIndices.Num();
+
+		for (int32 Idx = 0; Idx < InGraph.OutputIndices.Num(); ++Idx)
+		{
+			const int32 TensorIdx = InGraph.OutputIndices[Idx];
+			AddOutput(TensorIdx);
+		}
+
+		Operators.Reset(InGraph.Operators.Num());
+		for (const FOpDesc& OpDesc : InGraph.Operators)
+		{
+			AddOp(OpDesc, InGraph);
+		}
+
+		//Note: Validate edges here
+		
+		return true;
+	}
+
+	void AddInput(int32 TensorIdx)
+	{
+		AddEdge(
+			FEdge(EEdgeType::Input)
+				.SetTensorIdx(TensorIdx)
+				.SetNodeSrcOutput(NumInputs)
+		);
+
+		++NumInputs;
+	}
+
+	void AddOutput(int32 TensorIdx)
+	{
+		AddEdge(
+			FEdge(EEdgeType::Output)
+				.SetTensorIdx(TensorIdx)
+				.SetNodeDstInput(NumOutputs)
+		);
+
+		++NumOutputs;
+	}
+
+	void AddIntermediate(int32 TensorIdx, int32 NodeSrc, int32 NodeSrcOutput)
+	{
+		FEdge* ConnEdge = Edges.FindByPredicate(
+				[TensorIdx](const FEdge& Curr)
+				{
+					return Curr.TensorIdx == TensorIdx;
+				}
+		);
+
+		if (ConnEdge)
+		{
+			ConnectEdgeSrc(TensorIdx, NodeSrc, NodeSrcOutput);
+		}
+		else
+		{
+			AddEdge(
+				FEdge(EEdgeType::Intermediate)
+					.SetTensorIdx(TensorIdx)
+					.SetNodeSrc(NodeSrc)
+					.SetNodeSrcOutput(NodeSrcOutput)
+			);
+		}
+	}
+
+	void AddEdge(const FEdge& Edge)
+	{
+		FEdge* StartEdge = Edges.FindByPredicate(
+				[Edge](const FEdge& Curr)
+				{
+					return Curr.TensorIdx == Edge.TensorIdx;
+				}
+		);
+
+		check(StartEdge == nullptr);
+		Edges.Add(Edge);
+	}
+
+	bool ConnectEdgeDst(int32 TensorIdx, int32 NodeDst, int32 NodeDstInput)
+	{
+		FEdge* StartEdge = Edges.FindByPredicate(
+				[TensorIdx](const FEdge& Curr)
+				{
+					return Curr.TensorIdx == TensorIdx;
+				}
+		);
+
+		bool bFoundEdge = false;
+
+		for (FEdge* Curr = StartEdge; Curr->TensorIdx == TensorIdx; ++Curr)
+		{
+			if (Curr->NodeDst == -1 && Curr->NodeDstInput == -1)
+			{
+				Curr->NodeDst = NodeDst;
+				Curr->NodeDstInput = NodeDstInput;
+				bFoundEdge = true;
+				break;
+			}
+			else if (Curr->NodeDst == NodeDst && Curr->NodeDstInput == NodeDstInput)
+			{
+				bFoundEdge = true;
+				break;
+			}
+		}
+
+		checkf(bFoundEdge, TEXT("ConnectEdgeDst() has failed"));
+		return bFoundEdge;
+	}
+
+	bool ConnectEdgeSrc(int32 TensorIdx, int32 NodeSrc, int32 NodeSrcOutput)
+	{
+		FEdge* StartEdge =
+			Edges.FindByPredicate(
+				[TensorIdx](const FEdge& Curr)
+				{
+					return Curr.TensorIdx == TensorIdx;
+				}
+		);
+
+		bool bFoundEdge = false;
+
+		for (FEdge* Curr = StartEdge; Curr->TensorIdx == TensorIdx; ++Curr)
+		{
+			if (Curr->NodeSrc == -1 && Curr->NodeSrcOutput == -1)
+			{
+				Curr->NodeSrc = NodeSrc;
+				Curr->NodeSrcOutput = NodeSrcOutput;
+				bFoundEdge = true;
+				break;
+			}
+			else if (Curr->NodeSrc == NodeSrc && Curr->NodeSrcOutput == NodeSrcOutput)
+			{
+				bFoundEdge = true;
+				break;
+			}
+		}
+
+		checkf(bFoundEdge, TEXT("ConnectEdgeSrc() has failed"));
+		return bFoundEdge;
+	}
+
+	void AddOp(const FOpDesc& InOp, const FGraphDesc& InGraph)
+	{
+		DML_OPERATOR_GRAPH_NODE_DESC& OpDesc = Operators.Add_GetRef({});
+
+		OpDesc.Operator = InOp.Op->GetOperator();
+		OpDesc.Name = InOp.DbgName.Get();
+
+		const int32 NodeIdx = Operators.Num() - 1;
+
+		for (int32 Idx = 0; Idx < InOp.InputCount; ++Idx)
+		{
+			const int32 TensorIdx = InGraph.OpInputIndices[Idx + InOp.InputStart];				
+			
+			ConnectEdgeDst(TensorIdx, NodeIdx, Idx);
+		}
+
+		for (int32 Idx = 0; Idx < InOp.OutputCount; ++Idx)
+		{
+			const int32 TensorIdx = InGraph.OpOutputIndices[Idx + InOp.OutputStart];
+
+			AddIntermediate(TensorIdx, NodeIdx, Idx);
+		}
+	}
+
+	TArray<FEdge>							Edges;
+	TArray<DML_OPERATOR_GRAPH_NODE_DESC>	Operators;
+	int32									NumInputs;
+	int32									NumOutputs;
+};
+
+//
+//
+//
+FModel::FModel()
+{
+}
+
+//
+//
+//
+FModel::~FModel()
+{
+}
+
+//
+//
+//
+bool FModel::Init(TConstArrayView<uint8> ModelData, FDmlDeviceContext* InDevCtx)
+{
+	check(ModelData.Num() > 0);
+	FNNERuntimeFormat	Format;
+	int32 GuidSize = FModelInfo::Get()->GetGuidSize();
+	int32 VersionSize = FModelInfo::Get()->GetVersionSize();
+
+	if (!LoadModel(ModelData, Format, GuidSize + VersionSize))
+	{
+		return false;
+	}
+
+	DevCtx = InDevCtx;
+	DynamicRHI = GetID3D12PlatformDynamicRHI();
+
+	HRESULT Res = DevCtx->Device->CreateOperatorInitializer(0, nullptr, DML_PPV_ARGS(&OpInit));
+	if (!OpInit)
+	{
+		UE_LOG(LogNNE, Warning, TEXT("Error:Failed to create DML operator initializer"));
+		return false;
+	}
+
+	// DirectML requires all tensors to be concrete
+	// Notes: to handle dynamic tensor desc, op should init from symbolic shapes
+	TArray<NNECore::Internal::FTensor>	Tensors;
+
+	Tensors.Reset(AllSymbolicTensorDescs.Num());
+	for (const NNECore::FTensorDesc& TensorDesc : AllSymbolicTensorDescs)
+	{
+		Tensors.Emplace(NNECore::Internal::FTensor::MakeFromSymbolicDesc(TensorDesc));
+	}
+
+	FGraphBuilder				DmlGraphBuilder;
+	FGraphBuilder::FGraphDesc	DmlGraphDesc;
+
+	DmlGraphDesc.AllTensors = Tensors;
+	DmlGraphDesc.InputIndices = InputTensorIndices;
+	DmlGraphDesc.OutputIndices = OutputTensorIndices;
+	DmlGraphDesc.WeightIndices = WeightTensorIndices;
+	DmlGraphDesc.IntermediateIndices = IntermediateTensorIndices;
+	DmlGraphDesc.WeightTensors = WeightTensorRDGs;
+	
+	TArray<FGraphBuilder::FOpDesc>	DmlGraphOperators;
+	TArray<int32>					OpInputIndices;
+	TArray<int32>					OpOutputIndices;
+	uint64							TensorDataSize = 0;
+
+	// Loop over all operators in the model and create them
+	for (int32 Idx = 0; Idx < Format.Operators.Num(); ++Idx)
+	{
+		const FString TypeName = Format.Operators[Idx].TypeName;
+
+		FGraphBuilder::FOpDesc		OpDesc;
+		TArray<NNECore::Internal::FTensor>				OpInputTensors;
+		TArray<NNECore::Internal::FTensor>				OpOutputTensors;
+		NNECore::FAttributeMap	AttributeMap;
+
+		OpDesc.InputStart = OpInputIndices.Num();
+		OpDesc.OutputStart = OpOutputIndices.Num();
+
+		for (int32 InputTensorIndex : Format.Operators[Idx].InTensors)
+		{
+			const int32 WeightTensorIdx = WeightTensorIndices.Find(InputTensorIndex);
+
+			if (WeightTensorIdx >= 0)
+			{
+				TConstArrayView<uint8> TensorData = OpInputTensors.Emplace_GetRef(WeightTensorRDGs[WeightTensorIdx]).GetPreparedData<uint8>();
+				
+				TensorDataSize += Align(TensorData.Num(), DML_MINIMUM_BUFFER_TENSOR_ALIGNMENT);
+			}
+			else
+			{
+				NNECore::FTensorDesc SymbolicTensorDesc = AllSymbolicTensorDescs[InputTensorIndex];
+				// Notes: to handle dynamic tensor desc, op should init from symbolic shapes
+				OpInputTensors.Emplace(NNECore::Internal::FTensor::MakeFromSymbolicDesc(SymbolicTensorDesc));
+			}
+
+			OpInputIndices.Emplace(InputTensorIndex);
+		}
+
+		for (int32 OutputTensorIndex : Format.Operators[Idx].OutTensors)
+		{
+			NNECore::FTensorDesc SymbolicTensorDesc = AllSymbolicTensorDescs[OutputTensorIndex];
+			// Notes: to handle dynamic tensor desc, op should init from symbolic shapes
+			OpOutputTensors.Emplace(NNECore::Internal::FTensor::MakeFromSymbolicDesc(SymbolicTensorDesc));
+			OpOutputIndices.Emplace(OutputTensorIndex);
+		}
+
+		for (const FNNEFormatAttributeDesc& Desc : Format.Operators[Idx].Attributes)
+		{
+			AttributeMap.SetAttribute(Desc.Name, Desc.Value);
+		}
+
+		OpDesc.Op = OpCreate(TypeName, OpInputTensors, OpOutputTensors, AttributeMap);
+
+		if (!OpDesc.Op)
+		{
+			UE_LOG(LogNNE, Warning, TEXT("Error:Failed to create operator:%s"), *TypeName);
+			return false;
+		}
+
+		OpDesc.InputCount = OpInputTensors.Num();
+		OpDesc.OutputCount = OpOutputTensors.Num();
+		OpDesc.DbgName = TypeName;
+
+		DmlGraphOperators.Emplace(OpDesc);
+	}
+
+	DmlGraphDesc.Operators = DmlGraphOperators;
+	DmlGraphDesc.OpInputIndices = OpInputIndices;
+	DmlGraphDesc.OpOutputIndices = OpOutputIndices;
+
+	CompiledOp = DmlGraphBuilder.Compile(DevCtx, DmlGraphDesc);
+	if (!CompiledOp)
+	{
+		return false;
+	}
+
+	return InitCompiledOp(OpInputIndices, TensorDataSize);
+}
+
+//
+//
+//
+bool FModel::InitCompiledOp(TConstArrayView<int32> OpInputIndices, uint64 TensorDataSize)
+{
+	static constexpr EBufferUsageFlags	WeightBuffUsage = BUF_UnorderedAccess;
+	static constexpr ERHIAccess			WeightBuffAccess = ERHIAccess::UAVMask;
+
+	static constexpr EBufferUsageFlags	PersistBuffFlags = BUF_Static | BUF_ShaderResource | BUF_UnorderedAccess;
+	static constexpr ERHIAccess			PersistBuffAccess = ERHIAccess::UAVMask;
+
+	static constexpr EBufferUsageFlags	TempBuffFlags = BUF_Volatile | BUF_UnorderedAccess;
+	static constexpr ERHIAccess			TempBuffAccess = ERHIAccess::UAVMask;
+
+	HRESULT					Res;
+	IDMLDevice*				Device = DevCtx->Device;
+	IDMLCompiledOperator*	CompiledOps[] = { CompiledOp };
+	
+	Res = OpInit->Reset(UE_ARRAY_COUNT(CompiledOps), CompiledOps);
+	if (FAILED(Res))
+	{
+		UE_LOG(LogNNE, Warning, TEXT("Error:Failed to reset DirectML operator initializer"));
+		return false;
+	}
+
+	DML_BINDING_PROPERTIES InitBindProps = OpInit->GetBindingProperties();
+	DML_BINDING_PROPERTIES ExecBindProps = CompiledOp->GetBindingProperties();
+
+	DescCount = std::max(InitBindProps.RequiredDescriptorCount, ExecBindProps.RequiredDescriptorCount);
+	
+	D3D12_DESCRIPTOR_HEAP_DESC	HeapDesc = {};
+
+	HeapDesc.Flags = D3D12_DESCRIPTOR_HEAP_FLAG_SHADER_VISIBLE;
+	HeapDesc.Type = D3D12_DESCRIPTOR_HEAP_TYPE_CBV_SRV_UAV;
+	HeapDesc.NumDescriptors = DescCount;
+
+	Res = DevCtx->D3D12Device->CreateDescriptorHeap(&HeapDesc, IID_PPV_ARGS(&DescHeap));
+	if (!DescHeap)
+	{
+		UE_LOG(LogNNE, Warning, TEXT("Failed to create descriptor heap, res:%x"), Res);
+		return false;
+	}
+
+	DescSize = DevCtx->D3D12Device->GetDescriptorHandleIncrementSize(D3D12_DESCRIPTOR_HEAP_TYPE_CBV_SRV_UAV);
+
+	BindingTable.Reset(new FBindingTable());
+	if (!BindingTable->Init(this))
+	{
+		return false;
+	}
+
+	MemSizeTemp = ExecBindProps.TemporaryResourceSize;
+	MemSizePersist = ExecBindProps.PersistentResourceSize;
+
+	FEvent* Signal = FGenericPlatformProcess::GetSynchEventFromPool(false);
+
+	ENQUEUE_RENDER_COMMAND(FModel_SetTensorData)
+	(
+		[
+			this, 
+			Signal, 
+			InitTempMemSize = InitBindProps.TemporaryResourceSize,
+			TensorDataSize
+		]
+		(FRHICommandListImmediate& RHICmdList)
+		{
+			FRHIBufferInputArray	Inputs;
+
+			for (int32 InputIdx : InputTensorIndices)
+			{
+				Inputs.Emplace(nullptr);
+			}
+
+			TArray<CD3DX12_RESOURCE_BARRIER, TInlineAllocator<MaxNumInputs>>	Barriers;
+			FGPUFenceRHIRef	UploadFence = nullptr;
+
+			if (TensorDataSize)
+			{
+				UploadFence = RHICreateGPUFence(TEXT("FInferenceModel_UploadFence"));
+				
+				FRHIBuffer*		UploadBuff = CreateRHIBuffer(RHICmdList, TensorDataSize, BUF_ShaderResource | BUF_Dynamic | BUF_FastVRAM, ERHIAccess::CopySrc, TEXT("FInferenceModel_UploadBuffer"));
+				uint8*			UploadBuffPtr = static_cast<uint8*>(RHICmdList.LockBuffer(UploadBuff, 0, TensorDataSize, RLM_WriteOnly_NoOverwrite));
+				uint64			UploadOffset = 0;
+
+				for (const FTensorRDG& Tensor : WeightTensorRDGs)
+				{
+					TConstArrayView<uint8>	TensorData = Tensor.GetPreparedData<uint8>();
+
+					FBufferRHIRef WeightBuff;
+
+					WeightBuff = CreateRHIBuffer(RHICmdList, TensorData.Num(), WeightBuffUsage, WeightBuffAccess, TEXT("FModel_TensorWeights"));
+				
+					FMemory::Memcpy(UploadBuffPtr + UploadOffset, TensorData.GetData(), TensorData.Num());
+					RHICmdList.CopyBufferRegion(WeightBuff, 0, UploadBuff, UploadOffset, TensorData.Num());
+					UploadOffset += Align(TensorData.Num(), DML_MINIMUM_BUFFER_TENSOR_ALIGNMENT);
+
+					Inputs.Emplace(WeightBuff);
+				
+					Barriers.Emplace(
+						CD3DX12_RESOURCE_BARRIER::Transition(
+							DynamicRHI->RHIGetResource(WeightBuff),
+							D3D12_RESOURCE_STATE_COPY_DEST,
+							D3D12_RESOURCE_STATE_UNORDERED_ACCESS)
+					);
+				}
+
+				RHICmdList.UnlockBuffer(UploadBuff);
+				RHICmdList.WriteGPUFence(UploadFence);
+			}
+
+			if (MemSizePersist)
+			{
+#ifdef NNE_USE_D3D12_RESOURCES
+				PersistBuff = CreateD3D12Buffer(MemSizePersist);
+#else		
+				PersistBuff = CreateRHIBuffer(RHICmdList, MemSizePersist, PersistBuffFlags, PersistBuffAccess, TEXT("FModelDml_PeristBuff"));
+#endif
+			}
+
+			if (MemSizeTemp)
+			{
+#ifdef NNE_USE_D3D12_RESOURCES
+				TempBuff = CreateD3D12Buffer(MemSizeTemp);
+#else
+				TempBuff = CreateRHIBuffer(RHICmdList, MemSizeTemp, TempBuffFlags, TempBuffAccess, TEXT("FModelDml_TempBuff"));
+#endif
+			}
+
+#ifdef NNE_USE_D3D12_RESOURCES
+			ID3D12Resource* InitTempBuff = nullptr;
+#else
+			FBufferRHIRef InitTempBuff;
+#endif
+
+			if (InitTempMemSize)
+			{
+#ifdef NNE_USE_D3D12_RESOURCES
+				InitTempBuff = CreateD3D12Buffer(InitTempMemSize);
+#else
+				TempBuff = CreateRHIBuffer(RHICmdList, InitTempMemSize, TempBuffFlags, TempBuffAccess, TEXT("FModelDml_InitTempBuff"));
+#endif
+			}
+
+			RHICmdList.EnqueueLambda(
+				[this, Inputs, Barriers, InitTempBuff, UploadFence](FRHICommandListImmediate& RHICmdList)
+				{
+					while (UploadFence && UploadFence->NumPendingWriteCommands.GetValue() > 0)
+					{
+						FPlatformProcess::Sleep(0.001);
+					}
+
+					ID3D12GraphicsCommandList* D3DCmdList = nullptr;
+
+					D3DCmdList = DynamicRHI->RHIGetGraphicsCommandList(DevCtx->DeviceIndex);
+
+					BindingTable->Bind(OpInit, Inputs, PersistBuff, InitTempBuff);
+			
+					D3DCmdList = DynamicRHI->RHIGetGraphicsCommandList(DevCtx->DeviceIndex);
+
+					if (!Barriers.IsEmpty())
+					{
+						D3DCmdList->ResourceBarrier(Barriers.Num(), Barriers.GetData());
+					}
+					D3DCmdList->SetDescriptorHeaps(1, &DescHeap);
+					DevCtx->CmdRec->RecordDispatch(D3DCmdList, OpInit, BindingTable->Get());
+
+					DynamicRHI->RHIFinishExternalComputeWork(DevCtx->DeviceIndex, D3DCmdList);
+				}
+			);			
+			
+			RHICmdList.ImmediateFlush(EImmediateFlushType::FlushRHIThread);
+			Signal->Trigger();
+		}
+	);
+
+	Signal->Wait();
+	FGenericPlatformProcess::ReturnSynchEventToPool(Signal);
+
+	return true;
+}
+
+BEGIN_SHADER_PARAMETER_STRUCT(FTensorBufferParamsDml, )
+	RDG_BUFFER_ACCESS(Buffer, ERHIAccess::UAVCompute)
+END_SHADER_PARAMETER_STRUCT()
+
+//
+//
+//
+void FModel::AddDispatchOps_RenderThread(FRDGBuilder& GraphBuilder)
+{
+	const ERDGPassFlags TransitionBuffFlags = ERDGPassFlags::Compute | ERDGPassFlags::NeverCull;
+
+	InputBuffers.SetNumUninitialized(InputTensorIndices.Num() + WeightTensorIndices.Num());
+	OutputBuffers.SetNumUninitialized(OutputTensorIndices.Num());
+
+	for (int32 Idx = 0; Idx < InputTensorIndices.Num(); ++Idx)
+	{
+		FTensorBufferParamsDml* Params = GraphBuilder.AllocParameters<FTensorBufferParamsDml>();
+		Params->Buffer = AllTensorRDGRefs[InputTensorIndices[Idx]]->GetBuffer();
+
+		GraphBuilder.AddPass(
+			RDG_EVENT_NAME("FModel_Dispatch_GetInputBuffer"),
+			Params,
+			TransitionBuffFlags,
+			[this, Idx, Params](FRHICommandListImmediate& RHICmdList)
+			{
+				InputBuffers[Idx] = Params->Buffer->GetRHI();
+			}
+		);
+	}
+
+	for (int32 Idx = 0; Idx < WeightTensorIndices.Num(); ++Idx)
+	{
+		InputBuffers[Idx + InputTensorIndices.Num()] = nullptr;
+	}
+
+	for (int32 Idx = 0; Idx < OutputTensorIndices.Num(); ++Idx)
+	{
+		FTensorBufferParamsDml* Params = GraphBuilder.AllocParameters<FTensorBufferParamsDml>();
+		Params->Buffer = AllTensorRDGRefs[OutputTensorIndices[Idx]]->GetBuffer();
+
+		GraphBuilder.AddPass(
+			RDG_EVENT_NAME("FModel_Dispatch_GetOutputBuffer"),
+			Params,
+			TransitionBuffFlags,
+			[this, Idx, Params](FRHICommandListImmediate& RHICmdList)
+			{
+				OutputBuffers[Idx] = Params->Buffer->GetRHI();
+			}
+		);
+	}
+	
+	GraphBuilder.AddPass(
+		RDG_EVENT_NAME("FModel_Dispatch"),
+		ERDGPassFlags::None | ERDGPassFlags::NeverCull,
+		[this](FRHICommandListImmediate& RHICmdList)
+		{
+			RHICmdList.EnqueueLambda(
+				[this](FRHICommandListImmediate& RHICmdList)
+				{
+					TArray<CD3DX12_RESOURCE_BARRIER, TInlineAllocator<MaxNumInputs + MaxNumOutputs>>	Barriers;
+
+					for (FRHIBuffer* Buffer : InputBuffers)
+					{
+						if (!Buffer)
+						{
+							continue;
+						}
+
+						ID3D12Resource* Resource = DynamicRHI->RHIGetResource(Buffer);
+
+						//Note: We should not assume COPY_DEST state
+						Barriers.Emplace(
+							CD3DX12_RESOURCE_BARRIER::Transition(
+								Resource,
+								D3D12_RESOURCE_STATE_COPY_DEST,
+								D3D12_RESOURCE_STATE_UNORDERED_ACCESS)
+						);
+					}
+
+					//Note: We should use this instead of NNE_USE_D3D12_RESOURCES
+					//FBufferRHIRef TempBuff = MemSizeTemp ? CreateRHIBuffer(RHICmdList, MemSizeTemp, TempBuffUsage, TempBuffAccess, TEXT("FModel_Dispatch_TempBuff")) : nullptr;
+
+					BindingTable->Bind(CompiledOp, InputBuffers, OutputBuffers, PersistBuff, TempBuff);
+
+					ID3D12GraphicsCommandList* D3DCmdList = nullptr;
+
+					D3DCmdList = DynamicRHI->RHIGetGraphicsCommandList(DevCtx->DeviceIndex);
+					D3DCmdList->SetDescriptorHeaps(1, &DescHeap);
+					D3DCmdList->ResourceBarrier(Barriers.Num(), Barriers.GetData());
+					DevCtx->CmdRec->RecordDispatch(D3DCmdList, CompiledOp, BindingTable->Get());
+
+					DynamicRHI->RHIFinishExternalComputeWork(DevCtx->DeviceIndex, D3DCmdList);
+				}
+			);
+		}
+	);
+}
+
+//
+// Create operator
+//
+FOperatorDml* FModel::OpCreate(const FString& OpName, TArrayView<const NNECore::Internal::FTensor> InputTensorDescs, TArrayView<const NNECore::Internal::FTensor> OutputTensorDescs, const NNECore::FAttributeMap& Attributes)
+{
+	FOperatorRegistryDml::OperatorCreateFunc CreateFn = FOperatorRegistryDml::Get()->OpFind(OpName);
+
+	if (!CreateFn)
+	{
+		UE_LOG(LogNNE, Warning, TEXT("Dml MLOperatorRegistry failed to find operator:%s"), *OpName);
+		return nullptr;
+	}
+
+	FOperatorDml* Op = CreateFn();
+
+	if (!Op->Initialize(DevCtx->Device, InputTensorDescs, OutputTensorDescs, Attributes))
+	{
+		delete Op;
+
+		UE_LOG(LogNNE, Warning, TEXT("Error:Failed to initialize operator:%s"), *OpName);
+		return nullptr;
+	}
+
+	Op->GetOperator()->SetName(*OpName);
+
+	return Op;
+}
+
+FBufferRHIRef FModel::CreateRHIBuffer(FRHICommandListImmediate& RHICmdList, uint32 Size, EBufferUsageFlags Usage, ERHIAccess Access, const TCHAR* DbgName)
+{
+	FBufferRHIRef Buff = nullptr;
+
+	if (Size)
+	{
+		FRHIResourceCreateInfo CreateInfo(DbgName);
+		
+		Buff = RHICmdList.CreateBuffer(Size, Usage, 1, Access, CreateInfo);
+	}
+
+	check(Buff);
+	return Buff;
+}
+
+ID3D12Resource* FModel::CreateD3D12Buffer(uint32 Size, D3D12_RESOURCE_STATES ResourceState, D3D12_HEAP_TYPE HeapType, const TCHAR* DebugName)
+{
+	ID3D12Resource* Resource = nullptr;
+
+	D3D12_RESOURCE_DESC ResourceDesc = CD3DX12_RESOURCE_DESC::Buffer(Size,D3D12_RESOURCE_FLAG_ALLOW_UNORDERED_ACCESS);
+	CD3DX12_HEAP_PROPERTIES HeapProps = CD3DX12_HEAP_PROPERTIES(HeapType);
+	HRESULT Res;
+
+	Res = DevCtx->D3D12Device->CreateCommittedResource(
+		&HeapProps,
+		D3D12_HEAP_FLAG_NONE,
+		&ResourceDesc,
+		ResourceState,
+		nullptr,
+		IID_PPV_ARGS(&Resource));
+
+	if (FAILED(Res))
+	{
+		UE_LOG(LogNNE, Warning, TEXT("Error:FInferenceModel failed to create D3D12 resource"));
+		return nullptr;
+	}
+
+	if (Resource && DebugName)
+	{
+		Resource->SetName(DebugName);
+	}
+
+	return Resource;
+}
+
+//
+//
+//
+int FModel::PrepareTensorShapesAndData()
+{
+	for (NNECore::FTensorDesc SymbolicTensorDesc : AllSymbolicTensorDescs)
+	{
+		if (!SymbolicTensorDesc.GetShape().IsConcrete())
+		{
+			UE_LOG(LogNNE, Warning, TEXT("DML runtime does not support model with variable shapes yet."));
+			return -1;
+		}
+	}
+
+	return 0;
+}
+
+#endif // NNE_USE_DIRECTML
+
+} // namespace UE::NNERuntimeRDG::Private::Dml
+
+
diff --git a/Engine/Plugins/Experimental/NNE/Source/NNERuntimeRDG/Private/Dml/NNEDmlModel.h b/Engine/Plugins/Experimental/NNE/Source/NNERuntimeRDG/Private/Dml/NNEDmlModel.h
new file mode 100644
--- /dev/null
+++ b/Engine/Plugins/Experimental/NNE/Source/NNERuntimeRDG/Private/Dml/NNEDmlModel.h
@@ -0,0 +1,102 @@
+// Copyright Epic Games, Inc. All Rights Reserved.
+
+#pragma once
+
+#include "NNEDmlCommon.h"
+#include "NNERuntimeRDGModel.h"
+
+#define NNE_USE_D3D12_RESOURCES
+
+struct ID3D12DynamicRHI;
+
+namespace UE::NNERuntimeRDG::Private::Dml
+{
+
+class FModelInfo
+{
+public:
+
+	static FModelInfo* Get();
+
+	FGuid GetGuid();
+	int32 GetVersion();
+
+	int32 GetGuidSize();
+	int32 GetVersionSize();
+
+	bool ValidateGuidAndVersion(const uint8* InGuid, const uint8* InVersion);
+
+private:
+
+	FModelInfo();
+
+	FGuid	Guid;
+	int32	Version;
+};
+
+#ifdef NNE_USE_DIRECTML
+
+class FDmlDeviceContext;
+class FOperatorDml;
+
+//
+//
+//
+class FModel : public FModelRDG
+{
+	class FGraphBuilder;
+	class FBindingTable;
+	class FDebugName;
+
+public:
+
+	FModel();
+	~FModel();
+
+	bool Init(TConstArrayView<uint8> ModelData, FDmlDeviceContext* InDevCtx);
+
+protected:
+
+	virtual void AddDispatchOps_RenderThread(FRDGBuilder& GraphBuilder) override;
+	virtual int PrepareTensorShapesAndData() override;
+
+private:
+
+	bool InitCompiledOp(TConstArrayView<int32> OpInputIndices, uint64 TensorDataSize);
+
+	FOperatorDml* OpCreate(const FString& Name, TArrayView<const NNECore::Internal::FTensor> InputTensorDesc, TArrayView<const NNECore::Internal::FTensor> OutputTensorDescs, const NNECore::FAttributeMap& Attributes);
+
+	FBufferRHIRef CreateRHIBuffer(FRHICommandListImmediate& RHICmdList, uint32 Size, EBufferUsageFlags Usage, ERHIAccess Access, const TCHAR* DbgName);
+	ID3D12Resource* CreateD3D12Buffer(uint32 Size, D3D12_RESOURCE_STATES ResourceState = D3D12_RESOURCE_STATE_COMMON, D3D12_HEAP_TYPE HeapType = D3D12_HEAP_TYPE_DEFAULT, const TCHAR* DebugName = nullptr);
+
+	//Note: This should go into RDG
+	static constexpr int32 MaxNumInputs = 32;
+	static constexpr int32 MaxNumOutputs = 4;
+
+	using FRHIBufferInputArray = TArray<FRHIBuffer*, TInlineAllocator<MaxNumInputs>>;
+	using FRHIBufferOutputArray = TArray<FRHIBuffer*, TInlineAllocator<MaxNumOutputs>>;
+
+	TComPtr<IDMLOperatorInitializer>	OpInit;
+	TComPtr<IDMLCompiledOperator>		CompiledOp;
+	FDmlDeviceContext*					DevCtx;
+	TUniquePtr<FBindingTable>			BindingTable;
+	TComPtr<ID3D12DescriptorHeap>		DescHeap;
+	uint32								DescCount;
+	uint32								DescSize;
+
+	FRHIBufferInputArray				InputBuffers;
+	FRHIBufferOutputArray				OutputBuffers;
+#ifdef NNE_USE_D3D12_RESOURCES
+	TComPtr<ID3D12Resource>				PersistBuff;
+	TComPtr<ID3D12Resource>				TempBuff;
+#else
+	FBufferRHIRef						PersistBuff;
+#endif
+	uint64								MemSizeTemp;
+	uint64								MemSizePersist;
+	ID3D12DynamicRHI*					DynamicRHI;
+};
+
+#endif // NNE_USE_DIRECTML
+
+} // namespace UE::NNERuntimeRDG::Private::Dml
diff --git a/Engine/Plugins/Experimental/NNE/Source/NNERuntimeRDG/Private/Dml/NNEDmlOperator.cpp b/Engine/Plugins/Experimental/NNE/Source/NNERuntimeRDG/Private/Dml/NNEDmlOperator.cpp
new file mode 100644
--- /dev/null
+++ b/Engine/Plugins/Experimental/NNE/Source/NNERuntimeRDG/Private/Dml/NNEDmlOperator.cpp
@@ -0,0 +1,259 @@
+// Copyright Epic Games, Inc. All Rights Reserved.
+
+#ifdef NNE_USE_DIRECTML
+
+#include "NNEDmlOperator.h"
+#include "ID3D12DynamicRHI.h"
+
+namespace UE::NNERuntimeRDG::Private::Dml
+{
+
+namespace DmlUtil
+{
+	void SetTensorStrides(FTensorDesc& TensorDesc, const NNECore::Internal::FTensor& InputDesc)
+	{
+		uint32 CurrStride = 1;
+
+		TensorDesc.Strides.SetNum(InputDesc.GetShape().Rank());
+		
+		for (int32 i = InputDesc.GetShape().Rank() - 1; i >= 0; --i)
+		{
+			TensorDesc.Strides[i] = CurrStride;
+			CurrStride *= InputDesc.GetShape().GetData()[i];
+		}
+	}
+
+	void SetTensorSizesAndStridesForBroadcast(FTensorDesc& TensorDesc, const NNECore::Internal::FTensor& InputDesc, const NNECore::Internal::FTensor& TargetDesc)
+	{
+		static_assert(NNECore::FTensorShape::MaxRank <= 8);
+		
+		const uint32 TargetDimension = TargetDesc.GetShape().Rank() != -1 ? TargetDesc.GetShape().Rank() : InputDesc.GetShape().Rank();
+		checkf(TargetDesc.GetShape().Rank() >= InputDesc.GetShape().Rank(), TEXT("Can't broadcast tensor from rank %d to rank %d, should be inferior or equal."), InputDesc.GetShape().Rank(), TargetDimension);
+		
+		TensorDesc.Sizes.SetNum(TargetDimension);
+		TensorDesc.Strides.SetNum(TargetDimension);
+
+		const int32 DimensionOffset = int32(TargetDimension - InputDesc.GetShape().Rank());
+		
+		for (int32 i = 0; i < (int32) TargetDimension; ++i)
+		{
+			TensorDesc.Sizes[i] = i < DimensionOffset ? 1 : InputDesc.GetShape().GetData()[i - DimensionOffset];
+		}
+
+		uint32 CurrStride = 1;
+
+		for (int32 i = TargetDimension - 1; i >= 0; --i)
+		{
+			const bool bBroadcast = TensorDesc.Sizes[i] < TargetDesc.GetShape().GetData()[i];
+
+			TensorDesc.Strides[i] = bBroadcast ? 0 : CurrStride;
+			CurrStride *= TensorDesc.Sizes[i];
+
+			TensorDesc.Sizes[i] = TargetDesc.GetShape().GetData()[i];
+		}
+	}
+
+	bool IsSameShape(const NNECore::Internal::FTensor& Left, const NNECore::Internal::FTensor& Right)
+	{
+		if (Left.GetShape().Rank() != Right.GetShape().Rank())
+		{
+			return false;
+		}
+		
+		for (int32 Idx = 0; Idx < Left.GetShape().Rank(); ++Idx)
+		{
+			if (Left.GetShape().GetData()[Idx] != Right.GetShape().GetData()[Idx])
+			{
+				return false;
+			}
+		}
+
+		return true;
+	}
+
+	DML_TENSOR_DATA_TYPE GetTensorDataType(ENNETensorDataType DataType)
+	{
+		switch (DataType)
+		{
+		case ENNETensorDataType::Double:
+			return DML_TENSOR_DATA_TYPE_FLOAT64;
+
+		case ENNETensorDataType::Float:
+			return DML_TENSOR_DATA_TYPE_FLOAT32;
+
+		case ENNETensorDataType::Half:
+			return DML_TENSOR_DATA_TYPE_FLOAT16;
+
+		case ENNETensorDataType::UInt64:
+			return DML_TENSOR_DATA_TYPE_UINT64;
+
+		case ENNETensorDataType::UInt32:
+			return DML_TENSOR_DATA_TYPE_UINT32;
+
+		case ENNETensorDataType::UInt16:
+			return DML_TENSOR_DATA_TYPE_UINT16;
+
+		case ENNETensorDataType::UInt8:
+			return DML_TENSOR_DATA_TYPE_UINT8;
+
+		case ENNETensorDataType::Int64:
+			return DML_TENSOR_DATA_TYPE_INT64;
+
+		case ENNETensorDataType::Int32:
+			return DML_TENSOR_DATA_TYPE_INT32;
+
+		case ENNETensorDataType::Int16:
+			return DML_TENSOR_DATA_TYPE_INT16;
+
+		case ENNETensorDataType::Int8:
+			return DML_TENSOR_DATA_TYPE_INT8;
+
+		default:
+			return DML_TENSOR_DATA_TYPE_UNKNOWN;
+		}
+	}
+
+	uint64 CalculateBufferSize(const FTensorDesc& DmlTensor, const NNECore::Internal::FTensor& Desc)
+	{
+		uint64 ElemSizeInBytes = Desc.GetElemByteSize();
+		
+		if (ElemSizeInBytes == 0)
+		{
+			return 0;
+		}
+
+		uint64 MinSizeInBytes = 0;		
+		uint32 IndexOfLastElement = 0;
+
+		for (int32 i = 0; i < DmlTensor.Sizes.Num(); ++i)
+		{
+			IndexOfLastElement += (DmlTensor.Sizes[i] - 1) * DmlTensor.Strides[i];
+		}
+
+		MinSizeInBytes = (static_cast<uint64>(IndexOfLastElement) + 1) * ElemSizeInBytes;
+		
+		// Round up to the nearest 4 bytes.
+		MinSizeInBytes = (MinSizeInBytes + 3) & ~3ull;
+
+		return MinSizeInBytes;
+	}
+}
+
+//
+// DirectML operator base class
+//
+IDMLOperator* FOperatorDml::GetOperator()
+{
+	return DmlOp;
+}
+
+bool FOperatorDml::InitDmlTensorDesc(DmlUtil::FTensorDesc& DmlTensorDesc, const NNECore::Internal::FTensor& TensorDesc)
+{
+	DML_TENSOR_DATA_TYPE DmlDataType = DmlUtil::GetTensorDataType(TensorDesc.GetDataType());
+
+	if (DmlDataType == DML_TENSOR_DATA_TYPE_UNKNOWN)
+	{
+		DmlTensorDesc.BuffDesc = DML_BUFFER_TENSOR_DESC{};
+		DmlTensorDesc.Desc = DML_TENSOR_DESC{};
+
+		return false;
+	}
+
+	DmlTensorDesc.Sizes = TensorDesc.GetShape().GetData();
+	//Note: We should support tensor padding using strides defined in FTensorDesc
+	//DmlUtil::SetTensorStrides(DmlTensorDesc, TensorDesc.Strides);
+		
+	DML_BUFFER_TENSOR_DESC& BuffDesc = DmlTensorDesc.BuffDesc;
+
+	BuffDesc = DML_BUFFER_TENSOR_DESC{};
+	BuffDesc.DataType = DmlDataType;
+	BuffDesc.Flags = TensorDesc.HasPreparedData() ? DML_TENSOR_FLAG_OWNED_BY_DML : DML_TENSOR_FLAG_NONE;
+	BuffDesc.DimensionCount = TensorDesc.GetShape().Rank();
+	BuffDesc.Sizes = DmlTensorDesc.Sizes.GetData();
+	BuffDesc.Strides = nullptr;
+	BuffDesc.TotalTensorSizeInBytes = TensorDesc.GetDataSize();
+
+	DmlTensorDesc.Desc = DML_TENSOR_DESC{ DML_TENSOR_TYPE_BUFFER, &DmlTensorDesc.BuffDesc };
+
+	return true;
+}
+
+bool FOperatorDml::InitDmlTensorDesc(DmlUtil::FTensorDesc& DmlTensorDesc, const NNECore::Internal::FTensor& TensorDesc, const NNECore::Internal::FTensor& BroadcastDesc)
+{
+	DML_TENSOR_DATA_TYPE DmlDataType = DmlUtil::GetTensorDataType(TensorDesc.GetDataType());
+
+	if (DmlDataType == DML_TENSOR_DATA_TYPE_UNKNOWN)
+	{
+		DmlTensorDesc.BuffDesc = DML_BUFFER_TENSOR_DESC{};
+		DmlTensorDesc.Desc = DML_TENSOR_DESC{};
+
+		return false;
+	}
+
+	if (DmlUtil::IsSameShape(TensorDesc, BroadcastDesc))
+	{
+		DmlTensorDesc.Sizes = TensorDesc.GetShape().GetData();
+		DmlUtil::SetTensorStrides(DmlTensorDesc, TensorDesc);
+	}
+	else if (TensorDesc.GetShape().Rank() > BroadcastDesc.GetShape().Rank())
+	{
+		return false;
+	}
+	else
+	{
+		DmlUtil::SetTensorSizesAndStridesForBroadcast(DmlTensorDesc, TensorDesc, BroadcastDesc);
+	}
+
+	//UE_LOG(LogNNE, Warning, TEXT("DmlTensorDesc:%d,%d,%d -> %d,%d,%d"),
+	//	TensorDesc.Sizes[0],
+	//	TensorDesc.Dimension > 1 ? TensorDesc.Sizes[1] : 0,
+	//	TensorDesc.Dimension > 2 ? TensorDesc.Sizes[2] : 0,
+	//	DmlTensorDesc.Sizes[0],
+	//	DmlTensorDesc.Dimension > 1 ? DmlTensorDesc.Sizes[1] : 0,
+	//	DmlTensorDesc.Dimension > 2 ? DmlTensorDesc.Sizes[2] : 0
+	//);
+
+	//UE_LOG(LogNNE, Warning, TEXT("DmlTensorStrides:%d,%d,%d"),
+	//	DmlTensorDesc.Strides[0], DmlTensorDesc.Strides[1], DmlTensorDesc.Strides[2]
+	//);
+
+	check(DmlTensorDesc.Strides.Num() == DmlTensorDesc.Sizes.Num());
+		
+	DML_BUFFER_TENSOR_DESC& BuffDesc = DmlTensorDesc.BuffDesc;
+		
+	BuffDesc = DML_BUFFER_TENSOR_DESC{};
+
+	BuffDesc.DataType = DmlDataType;
+	BuffDesc.Flags = TensorDesc.HasPreparedData() ? DML_TENSOR_FLAG_OWNED_BY_DML : DML_TENSOR_FLAG_NONE;
+	BuffDesc.DimensionCount = DmlTensorDesc.Sizes.Num();
+	BuffDesc.Sizes = DmlTensorDesc.Sizes.GetData();
+	BuffDesc.Strides = DmlTensorDesc.Strides.GetData();
+	BuffDesc.TotalTensorSizeInBytes = DmlUtil::CalculateBufferSize(DmlTensorDesc, TensorDesc);
+		
+	DmlTensorDesc.Desc = DML_TENSOR_DESC{ DML_TENSOR_TYPE_BUFFER, &DmlTensorDesc.BuffDesc };
+
+	return true;
+}
+
+// Create DirectML operator	
+bool FOperatorDml::CreateOperator(IDMLDevice* Device, const DML_OPERATOR_DESC& DmlOpDesc)
+{
+	IDMLOperator* Op = nullptr;
+
+	HRESULT Res;
+
+	Res = Device->CreateOperator(&DmlOpDesc, DML_PPV_ARGS(&Op));
+	if (!Op)
+	{
+		UE_LOG(LogNNE, Warning, TEXT("Error:Failed to create DML operator, hres:%d"), Res);
+		return false;
+	}
+
+	DmlOp = Op;
+
+	return DmlOp.IsValid();
+}
+
+} // namespace UE::NNERuntimeRDG::Private::Dml
+
+#endif // NNE_USE_DIRECTML
diff --git a/Engine/Plugins/Experimental/NNE/Source/NNERuntimeRDG/Private/Dml/NNEDmlOperator.h b/Engine/Plugins/Experimental/NNE/Source/NNERuntimeRDG/Private/Dml/NNEDmlOperator.h
new file mode 100644
--- /dev/null
+++ b/Engine/Plugins/Experimental/NNE/Source/NNERuntimeRDG/Private/Dml/NNEDmlOperator.h
@@ -0,0 +1,89 @@
+// Copyright Epic Games, Inc. All Rights Reserved.
+
+#pragma once
+
+#ifdef NNE_USE_DIRECTML
+
+#include "NNEDmlCommon.h"
+#include "NNECoreTypes.h"
+#include "NNECoreTensor.h"
+#include "NNECoreOperator.h"
+#include "NNECoreAttributeMap.h"
+#include "NNERuntimeRDG.h"
+
+#define NNE_DML_REGISTER_OP(OpName) \
+struct FDmlOperator##OpName##Registrator \
+{ \
+	FDmlOperator##OpName##Registrator() \
+	{ \
+		FOperatorRegistryDml::Get()->OpAdd(TEXT(#OpName), FOperatorDml##OpName##::Create); \
+	} \
+}; \
+\
+static FDmlOperator##OpName##Registrator RegisterDmlOperator##OpName;
+
+
+namespace UE::NNERuntimeRDG::Private::Dml
+{
+//
+//
+//
+class FDmlDeviceContext
+{
+public:
+
+	uint32							DeviceIndex;
+	ID3D12Device*					D3D12Device{ nullptr }; // Borrowed reference from RHI
+	TComPtr<IDMLDevice>				Device{ nullptr };
+	TComPtr<IDMLCommandRecorder>	CmdRec{ nullptr };
+};
+
+namespace DmlUtil
+{
+	struct FTensorDesc
+	{
+		DML_BUFFER_TENSOR_DESC												BuffDesc;
+		DML_TENSOR_DESC														Desc;
+		TArray<uint32, TInlineAllocator<NNECore::FTensorShape::MaxRank>>	Sizes;
+		TArray<uint32, TInlineAllocator<NNECore::FTensorShape::MaxRank>>	Strides;
+	};
+
+	extern void SetTensorStrides(FTensorDesc& TensorDesc, const NNECore::Internal::FTensor& InputDesc);
+	extern void SetTensorSizesAndStridesForBroadcast(FTensorDesc& TensorDesc, const NNECore::Internal::FTensor& InputDesc, const NNECore::Internal::FTensor& TargetDesc);
+	extern bool IsSameShape(const NNECore::Internal::FTensor& Left, const NNECore::Internal::FTensor& Right);
+	DML_TENSOR_DATA_TYPE GetTensorDataType(ENNETensorDataType DataType);
+	uint64 CalculateBufferSize(const FTensorDesc& DmlTensor, const NNECore::Internal::FTensor& Desc);
+}
+
+//
+// DirectML operator base class
+//
+class FOperatorDml
+{
+public:
+
+	virtual ~FOperatorDml() = default;
+
+	virtual bool Initialize(IDMLDevice* Device, TArrayView<const NNECore::Internal::FTensor> InputTensors, TArrayView<const NNECore::Internal::FTensor> OutputTensors, const NNECore::FAttributeMap& Attributes) = 0;
+
+	IDMLOperator* GetOperator();
+
+protected:
+
+	bool InitDmlTensorDesc(DmlUtil::FTensorDesc& DmlTensorDesc, const NNECore::Internal::FTensor& TensorDesc);
+	bool InitDmlTensorDesc(DmlUtil::FTensorDesc& DmlTensorDesc, const NNECore::Internal::FTensor& TensorDesc, const NNECore::Internal::FTensor& BroadcastDesc);
+	
+	bool CreateOperator(IDMLDevice* Device, const DML_OPERATOR_DESC& DmlOpDesc);
+
+	TComPtr<IDMLOperator>		DmlOp;
+};
+
+/**
+ * DirectML ML operator registry
+ */
+using FOperatorRegistryDml = TOperatorRegistryRDG<FOperatorDml>;
+
+
+} // namespace UE::NNERuntimeRDG::Private::Dml
+
+#endif // NNE_USE_DIRECTML
diff --git a/Engine/Plugins/Experimental/NNE/Source/NNERuntimeRDG/Private/Dml/NNEDmlOperatorConvolution.cpp b/Engine/Plugins/Experimental/NNE/Source/NNERuntimeRDG/Private/Dml/NNEDmlOperatorConvolution.cpp
new file mode 100644
--- /dev/null
+++ b/Engine/Plugins/Experimental/NNE/Source/NNERuntimeRDG/Private/Dml/NNEDmlOperatorConvolution.cpp
@@ -0,0 +1,393 @@
+// Copyright Epic Games, Inc. All Rights Reserved.
+
+#ifdef NNE_USE_DIRECTML
+
+#include "NNEDmlOperator.h"
+
+namespace UE::NNERuntimeRDG::Private::Dml
+{
+/**
+ * Convolution
+ */
+template <DML_CONVOLUTION_DIRECTION Direction>
+class FOperatorDmlConv : public FOperatorDml
+{
+	static constexpr uint32_t NcdhwDimensionCount = 5;
+	static constexpr uint32_t NcdhwSpatialDimensionCount = 3;
+	static constexpr uint32_t NonspatialDimensionCount = 2; // The batch and channel dimensions of NCW, NCHW, NCDHW....
+
+	using FSmallArray = TArray<uint32, TInlineAllocator<NcdhwSpatialDimensionCount>>;
+	using FIntArray = TArray<int32>;
+	using FUIntArray = TArray<uint32>;
+	
+	enum EAutoPad
+	{
+		NOTSET,
+		SAME_UPPER,
+		SAME_LOWER,
+		VALID
+	};
+
+	static EAutoPad AutoPadFromString(FStringView StringVal) 
+	{
+		if (FCString::Stricmp(StringVal.GetData(), TEXT("NOTSET")) == 0) 
+		{
+			return EAutoPad::NOTSET;
+		}
+		else if (FCString::Stricmp(StringVal.GetData(), TEXT("SAME_UPPER")) == 0)
+		{
+			return EAutoPad::SAME_UPPER;
+		}
+		else if (FCString::Stricmp(StringVal.GetData(), TEXT("SAME_LOWER")) == 0)
+		{
+			return EAutoPad::SAME_LOWER;
+		}
+		else if (FCString::Stricmp(StringVal.GetData(), TEXT("VALID")) == 0)
+		{
+			return EAutoPad::VALID;
+		}
+		else
+		{
+			return EAutoPad::NOTSET;
+		}
+	}
+
+	struct FConvArgs
+	{
+		EAutoPad		AutoPad;
+		FSmallArray		StartPadding;
+		FSmallArray		EndPadding;
+		FIntArray		OutPadding;
+		FIntArray		Dilations;
+		FIntArray		Strides;
+		FSmallArray		OutputShape;
+		uint32			NumDimensions;
+		FSmallArray		WindowSize;
+		int32			Group;
+
+		FConvArgs() = default;
+
+		//
+		//
+		//
+		bool Init(const NNECore::FTensorShape& InputShape, const NNECore::FTensorShape& FilterShape, const NNECore::FAttributeMap& Attributes)
+		{
+			check(InputShape.Rank() > NonspatialDimensionCount);
+			check(FilterShape.Rank() == InputShape.Rank());
+			
+			NumDimensions = InputShape.Rank() - NonspatialDimensionCount;
+
+			const FNNEAttributeValue*	AttrStrides = Attributes.GetAttributeValue(TEXT("strides"));
+
+			if (AttrStrides)
+			{
+				Strides = AttrStrides->GetValue<FIntArray>();
+			}
+			else
+			{
+				Strides.Init(1, NumDimensions);
+			}
+
+			check(Strides.Num() == 0 || Strides.Num() == FilterShape.Rank() - NonspatialDimensionCount);
+
+			const FNNEAttributeValue*	AttrDilations = Attributes.GetAttributeValue(TEXT("dilations"));
+
+			if (AttrDilations)
+			{
+				Dilations = AttrDilations->GetValue<FIntArray>();
+			}
+			else
+			{
+				Dilations.Init(1, NumDimensions);
+			}
+
+			check(Dilations.Num() == 0 || Dilations.Num() == FilterShape.Rank() - NonspatialDimensionCount);
+			
+			for (int32 Dim = FilterShape.Rank() - NumDimensions; Dim < FilterShape.Rank(); ++Dim)
+			{
+				WindowSize.Add(FilterShape.GetData()[Dim]);
+			}
+
+			AutoPad = AutoPadFromString(*Attributes.GetValue<FString>(TEXT("auto_pad")));
+
+			if (AutoPad == EAutoPad::NOTSET)
+			{
+				const FNNEAttributeValue* AttrPads = Attributes.GetAttributeValue(TEXT("pads"));
+				TArray<int32>	Pads;
+
+				if (AttrPads)
+				{
+					Pads = AttrPads->GetValue<TArray<int32>>();
+				}
+				else
+				{
+					Pads.Init(0, 2 * NumDimensions);
+				}
+
+				for (uint32 Dim = 0; Dim < NumDimensions; ++Dim)
+				{
+					StartPadding.Add(Pads[Dim]);
+					EndPadding.Add(Pads[Dim + NumDimensions]);
+				}
+			}
+			else if (AutoPad == EAutoPad::VALID)
+			{
+				StartPadding.Init(0, NumDimensions);
+				EndPadding.Init(0, NumDimensions);
+			}
+			else
+			{
+				StartPadding.Init(0, NumDimensions);
+				EndPadding.Init(0, NumDimensions);
+
+				SetAutoPadding(InputShape.GetData(), FilterShape.GetData());
+			}
+
+			Group = Attributes.GetValueOrDefault<int>(TEXT("group"), 1);
+
+			if (Direction == DML_CONVOLUTION_DIRECTION_FORWARD)
+			{
+				OutPadding.Init(0, NumDimensions);
+			}
+			else
+			{
+				const FNNEAttributeValue* AttrOutPadding = Attributes.GetAttributeValue(TEXT("output_padding"));
+
+				if (AttrOutPadding)
+				{
+					OutPadding = AttrOutPadding->GetValue<TArray<int32>>();
+				}
+				else
+				{
+					OutPadding.Init(0, NumDimensions);
+				}
+			}
+
+			if (Direction == DML_CONVOLUTION_DIRECTION_FORWARD)
+			{
+				SetOutputShape(InputShape.GetData(), FilterShape.GetData());
+			}
+			else
+			{
+				const FNNEAttributeValue* AttrOutShape = Attributes.GetAttributeValue(TEXT("output_shape"));
+
+				if (AttrOutShape)
+				{
+					TArray<int32> OutShapeVal = AttrOutShape->GetValue<TArray<int32>>();
+
+					for (int32 Value : OutShapeVal)
+					{
+						OutputShape.Add(uint32(Value));
+					}
+				}
+				else
+				{
+					SetOutputShape(InputShape.GetData(), FilterShape.GetData());
+				}
+			}
+
+			return true;
+		}
+
+		//
+		//
+		//
+		void SetAutoPadding(TConstArrayView<uint32> InputShape, TConstArrayView<uint32> FilterShape)
+		{
+			const uint32 DimOffset = NonspatialDimensionCount;
+
+			for (uint32 Dim = 0; Dim < NumDimensions; ++Dim)
+			{
+				uint32 InputLen = uint32(InputShape[Dim + DimOffset]);
+				uint32 StridedOutLen = (InputLen + Strides[Dim] - 1) / Strides[Dim];
+				uint32 KernelLen = 1 + (WindowSize[Dim] - 1) * Dilations[Dim];
+				uint32 Len = Strides[Dim] * (StridedOutLen - 1) + KernelLen;
+				uint32 Padding = (Len <= InputLen) ? 0 : (Len - InputLen);
+
+				if (AutoPad == EAutoPad::SAME_LOWER)
+				{
+					StartPadding[Dim] = (Padding + 1) / 2;
+				}
+				else
+				{
+					StartPadding[Dim] = Padding / 2;
+				}
+
+				EndPadding[Dim] = Padding - StartPadding[Dim];
+			}
+		}
+
+		//
+		//
+		//
+		void SetOutputShape(TConstArrayView<uint32> InputShape, TConstArrayView<uint32> FilterShape)
+		{
+			const uint32 DimOffset = NonspatialDimensionCount;
+
+			OutputShape.SetNumUninitialized(InputShape.Num());
+
+			if (Direction == DML_CONVOLUTION_DIRECTION_FORWARD)
+			{
+				OutputShape[0] = InputShape[0];
+				OutputShape[1] = FilterShape[0];
+
+				for (uint32 Dim = 0; Dim < NumDimensions; ++Dim)
+				{
+					uint32 InputLen = InputShape[Dim + DimOffset];
+					uint32 PaddedLen = InputLen + StartPadding[Dim] + EndPadding[Dim];
+					uint32 KernelLen = 1 + (WindowSize[Dim] - 1) * Dilations[Dim];
+
+					checkf(KernelLen <= PaddedLen, TEXT("KernelLen must < PaddedLen"));
+					checkf(Strides[Dim] != 0, TEXT("Strides must be != 0"));
+
+					uint32 StridableOutLen = PaddedLen - KernelLen;
+					uint32 OutLen = 1 + (StridableOutLen / Strides[Dim]);
+
+					OutputShape[Dim + DimOffset] = OutLen;
+				}
+			}
+			else
+			{
+				OutputShape[0] = InputShape[0];
+				OutputShape[1] = FilterShape[1] * Group;
+
+				for (uint32 Dim = 0; Dim < NumDimensions; ++Dim)
+				{
+					uint32 Padding = StartPadding[Dim] + EndPadding[Dim];
+					uint32 KernelLen = 1 + (WindowSize[Dim] - 1) * Dilations[Dim];
+
+					OutputShape[Dim + DimOffset] = (InputShape[Dim + DimOffset] - 1) * Strides[Dim] + KernelLen + OutPadding[Dim] - Padding;
+				}
+			}
+		}
+	};
+
+public:
+
+	static FOperatorDml* Create()
+	{
+		return new FOperatorDmlConv();
+	}
+
+	//
+	//
+	//
+	virtual bool Initialize(IDMLDevice* Device, TArrayView<const NNECore::Internal::FTensor> InputTensors, TArrayView<const NNECore::Internal::FTensor> OutputTensors, const NNECore::FAttributeMap& Attributes) override
+	{
+		const NNECore::Internal::FTensor& InputTensor = InputTensors[0];
+		const NNECore::Internal::FTensor& FilterTensor = InputTensors[1];
+		
+		FConvArgs	Args;
+		
+		if (!Args.Init(InputTensor.GetShape(), FilterTensor.GetShape(), Attributes))
+		{
+			return false;
+		}
+
+		NNECore::Internal::FTensor OutputTensor = OutputTensors[0];
+
+		if (Direction == DML_CONVOLUTION_DIRECTION_FORWARD)
+		{
+			OutputTensor.SetShape(NNECore::FTensorShape::Make(Args.OutputShape));
+		}
+
+		// Initialize tensor descriptors
+		DmlUtil::FTensorDesc	DmlInputTensor{};
+		DmlUtil::FTensorDesc	DmlFilterTensor{};
+		DmlUtil::FTensorDesc	DmlBiasTensor{};
+		DmlUtil::FTensorDesc	DmlOutputTensor{};
+
+		if (!InitDmlTensorDesc(DmlInputTensor, InputTensor))
+		{
+			UE_LOG(LogNNE, Warning, TEXT("Failed to initialize tensor(s) for DML inference"));
+			return false;
+		}
+
+		if (!InitDmlTensorDesc(DmlFilterTensor, FilterTensor))
+		{
+			UE_LOG(LogNNE, Warning, TEXT("Failed to initialize tensor(s) for DML inference"));
+			return false;
+		}
+
+		if (InputTensors.Num() > 2)
+		{
+			NNECore::Internal::FTensor BiasTensor = InputTensors[2];
+
+			if (BiasTensor.GetShape().Rank() < (int32) Args.NumDimensions)
+			{
+				FSmallArray Shape;
+
+				Shape.Add(1);
+				Shape.Add(BiasTensor.GetShape().GetData()[0]);
+				Shape.Add(1);
+
+				for (int32 Dim = 3; Dim < InputTensor.GetShape().Rank(); ++Dim)
+				{
+					Shape.Add(1);
+				}
+
+				BiasTensor.SetShape(NNECore::FTensorShape::Make(Shape));
+			}
+
+			if (!InitDmlTensorDesc(DmlBiasTensor, BiasTensor))
+			{
+				UE_LOG(LogNNE, Warning, TEXT("Failed to initialize tensor(s) for DML inference"));
+				return false;
+			}
+		}
+
+		if (!InitDmlTensorDesc(DmlOutputTensor, OutputTensor))
+		{
+			UE_LOG(LogNNE, Warning, TEXT("Failed to initialize tensor(s) for DML inference"));
+			return false;
+		}
+
+		DML_CONVOLUTION_OPERATOR_DESC	DmlConvOpDesc{};
+
+		DmlConvOpDesc.InputTensor = &DmlInputTensor.Desc;
+		DmlConvOpDesc.FilterTensor = &DmlFilterTensor.Desc;
+		DmlConvOpDesc.BiasTensor = InputTensors.Num() > 2 ? &DmlBiasTensor.Desc : nullptr;
+		DmlConvOpDesc.OutputTensor = &DmlOutputTensor.Desc;
+		DmlConvOpDesc.Mode = DML_CONVOLUTION_MODE_CROSS_CORRELATION;
+		DmlConvOpDesc.Direction = Direction;
+		DmlConvOpDesc.DimensionCount = Args.NumDimensions;
+		DmlConvOpDesc.Strides = (uint32*) Args.Strides.GetData();
+		DmlConvOpDesc.Dilations = (uint32*) Args.Dilations.GetData();
+		DmlConvOpDesc.StartPadding = Args.StartPadding.GetData();
+		DmlConvOpDesc.EndPadding = Args.EndPadding.GetData();
+		DmlConvOpDesc.OutputPadding = (uint32*) Args.OutPadding.GetData();
+		DmlConvOpDesc.GroupCount = Args.Group;
+
+		DML_OPERATOR_DESC DmlOpDesc{};
+
+		DmlOpDesc.Type = DML_OPERATOR_CONVOLUTION;
+		DmlOpDesc.Desc = &DmlConvOpDesc;
+
+		return CreateOperator(Device, DmlOpDesc);
+	}
+};
+
+void RegisterConvOperator()
+{
+	FOperatorRegistryDml::Get()->OpAdd(TEXT("Conv"), FOperatorDmlConv<DML_CONVOLUTION_DIRECTION_FORWARD>::Create);
+}
+
+void RegisterConvTransposeOperator()
+{
+	FOperatorRegistryDml::Get()->OpAdd(TEXT("ConvTranspose"), FOperatorDmlConv<DML_CONVOLUTION_DIRECTION_BACKWARD>::Create);
+}
+
+struct FDmlOperatorConvRegistrator
+{
+	FDmlOperatorConvRegistrator()
+	{
+		RegisterConvOperator();
+		RegisterConvTransposeOperator();
+	}
+};
+
+static FDmlOperatorConvRegistrator RegisterDmlOperatorConv;
+
+} // namespace UE::NNERuntimeRDG::Private::Dml
+
+#endif // NNE_USE_DIRECTML
diff --git a/Engine/Plugins/Experimental/NNE/Source/NNERuntimeRDG/Private/Dml/NNEDmlOperatorElementwise.cpp b/Engine/Plugins/Experimental/NNE/Source/NNERuntimeRDG/Private/Dml/NNEDmlOperatorElementwise.cpp
new file mode 100644
--- /dev/null
+++ b/Engine/Plugins/Experimental/NNE/Source/NNERuntimeRDG/Private/Dml/NNEDmlOperatorElementwise.cpp
@@ -0,0 +1,435 @@
+// Copyright Epic Games, Inc. All Rights Reserved.
+#ifdef NNE_USE_DIRECTML
+#include "NNEDmlOperator.h"
+
+namespace UE::NNERuntimeRDG::Private::Dml
+{
+
+/**
+ * Utility function to get operator type
+ */
+template<typename TElementWiseOpDesc>
+DML_OPERATOR_TYPE GetDmlElementWiseUnaryOpType();
+
+#define OP_EW(OpName) \
+template<> \
+DML_OPERATOR_TYPE GetDmlElementWiseUnaryOpType<DML_ELEMENT_WISE_##OpName##_OPERATOR_DESC>() { return DML_OPERATOR_ELEMENT_WISE_##OpName; }
+
+#define OP_AN(OpName) \
+template<> \
+DML_OPERATOR_TYPE GetDmlElementWiseUnaryOpType<DML_ACTIVATION_##OpName##_OPERATOR_DESC>() { return DML_OPERATOR_ACTIVATION_##OpName; }
+
+
+OP_EW(IDENTITY)
+OP_EW(ABS)
+OP_EW(ACOS)
+OP_EW(ACOSH)
+OP_EW(ASIN)
+OP_EW(ASINH)
+OP_EW(ATAN)
+OP_EW(ATANH)
+// BitShift
+// Cast
+OP_EW(CEIL)
+OP_EW(CLIP)
+OP_EW(COS)
+OP_EW(COSH)
+OP_AN(ELU)
+OP_EW(ERF)
+OP_EW(EXP)
+OP_EW(FLOOR)
+OP_EW(IS_INFINITY)
+OP_EW(IS_NAN)
+OP_AN(HARDMAX)
+OP_AN(HARD_SIGMOID)
+OP_AN(LEAKY_RELU)
+OP_EW(LOG)
+OP_EW(NEGATE)
+// Not
+OP_EW(RECIP)
+OP_AN(RELU)
+OP_EW(ROUND)
+OP_AN(SCALED_ELU)
+OP_AN(SIGMOID)
+OP_EW(SIGN)
+OP_EW(SIN)
+OP_EW(SINH)
+OP_AN(SOFTPLUS)
+OP_AN(SOFTSIGN)
+OP_EW(SQRT)
+OP_EW(TAN)
+OP_EW(TANH)
+
+#undef OP_EW
+#undef OP_AN
+
+
+/**
+ * Utility function to get operator type
+ */
+//template<typename TActivationOpDesc>
+//DML_OPERATOR_TYPE GetDmlActivationOpType();
+//
+//#define OP(OpName) \
+//template<> \
+//DML_OPERATOR_TYPE GetDmlActivationOpType<DML_ACTIVATION_##OpName##_OPERATOR_DESC>() { return DML_OPERATOR_ACTIVATION_##OpName; }
+//
+//OP(ELU)
+//OP(ERF)
+//OP(HARDMAX)
+//OP(HARD_SIGMOID)
+//OP(LEAKY_RELU)
+//OP(LINEAR)
+//OP(LOG_SOFTMAX)
+//OP(PARAMETERIZED_RELU)
+//OP(PARAMETRIC_SOFTPLUS)
+//OP(RELU)
+//OP(SCALED_ELU)
+//OP(SCALED_TANH)
+//OP(SOFTMAX)
+//OP(SOFTPLUS)
+//OP(SOFTSIGN)
+//OP(TANH)
+//OP(THRESHOLDED_RELU)
+//
+
+/**
+ * Utility function to get operator type
+ */
+template<typename TElementWiseOpDesc>
+DML_OPERATOR_TYPE GetDmlElementWiseBinaryOpType();
+
+#define OP_EW(OpName) \
+template<> \
+DML_OPERATOR_TYPE GetDmlElementWiseBinaryOpType<DML_ELEMENT_WISE_##OpName##_OPERATOR_DESC>() { return DML_OPERATOR_ELEMENT_WISE_##OpName; }
+
+#define OP_AN(OpName) \
+template<> \
+DML_OPERATOR_TYPE GetDmlElementWiseBinaryOpType<DML_ACTIVATION_##OpName##_OPERATOR_DESC>() { return DML_OPERATOR_ACTIVATION_##OpName; }
+
+OP_EW(ADD)
+//OP_EW(LOGICAL_AND)
+OP_EW(DIVIDE)
+//OP_EW(LOGICAL_EQUALS)
+//OP_EW(LOGICAL_GREATER_THAN)
+//OP_EW(LOGICAL_LESS_THAN)
+//OP_EW(MOD)
+OP_EW(MULTIPLY)
+//OP_EW(LOGICAL_OR)
+OP_AN(PARAMETERIZED_RELU)
+OP_EW(POW)
+OP_EW(SUBTRACT)
+//OP_EW(LOGICAL_XOR)
+
+#undef OP_EW
+#undef OP_AN
+
+
+/**
+ * Element-wise unary ML operator implementation
+ */
+template
+<
+	typename DmlElementWiseOpDescType, 
+	NNECore::Internal::EElementWiseUnaryOperatorType OpType
+>
+class FOperatorDmlElementWiseUnary : public FOperatorDml
+{
+public:
+
+	static FOperatorDml* Create()
+	{
+		return new FOperatorDmlElementWiseUnary();
+	}
+
+	virtual ~FOperatorDmlElementWiseUnary() = default;
+
+private:
+
+	FOperatorDmlElementWiseUnary() : Alpha(0.0f), Beta(0.0f), Gamma(0.0f), Num(1) {}
+	float Alpha;
+	float Beta;
+	float Gamma;
+	uint32 Num;
+
+public:
+
+	//
+	//
+	//
+	virtual bool Initialize(IDMLDevice* Device, TArrayView<const NNECore::Internal::FTensor> InputTensors, TArrayView<const NNECore::Internal::FTensor> OutputTensors, const NNECore::FAttributeMap& Attributes) override
+	{
+		Num = InputTensors[0].GetVolume();
+
+		const NNECore::Internal::FTensor& InputTensorDesc = InputTensors[0];
+		const NNECore::Internal::FTensor& OutputTensorDesc = OutputTensors[0];
+
+		Alpha = Attributes.GetValueOrDefault(TEXT("alpha"), Alpha);
+		Beta = Attributes.GetValueOrDefault(TEXT("beta"), Beta);
+		Gamma = Attributes.GetValueOrDefault(TEXT("gamma"), Gamma);
+
+		// Initialize tensor descriptor (it's same for both input and output)
+		DmlUtil::FTensorDesc	DmlTensorDesc{};
+
+		if (!InitDmlTensorDesc(DmlTensorDesc, InputTensorDesc))
+		{
+			UE_LOG(LogNNE, Warning, TEXT("Failed to initialize tensor(s) for DML inference"));
+			return false;
+		}
+
+		DmlElementWiseOpDescType	DmlElemWiseOpDesc{};
+
+		InitDmlOpDesc(DmlElemWiseOpDesc, DmlTensorDesc);
+
+		DML_OPERATOR_DESC DmlOpDesc{};
+
+		DmlOpDesc.Type = GetDmlElementWiseUnaryOpType<DmlElementWiseOpDescType>();
+		DmlOpDesc.Desc = &DmlElemWiseOpDesc;
+
+		return CreateOperator(Device, DmlOpDesc);
+	}
+
+private:
+
+	template<typename OpDesc>
+	void InitDmlOpDesc(OpDesc& Desc, DmlUtil::FTensorDesc& TensorDesc)
+	{
+		Desc.InputTensor = &TensorDesc.Desc;
+		Desc.OutputTensor = &TensorDesc.Desc;
+	}
+
+	void InitDmlOpDesc(DML_ACTIVATION_SOFTPLUS_OPERATOR_DESC& Desc, DmlUtil::FTensorDesc& TensorDesc)
+	{
+		Desc.InputTensor = &TensorDesc.Desc;
+		Desc.OutputTensor = &TensorDesc.Desc;
+		Desc.Steepness = 1.0f;
+	}
+
+	void InitDmlOpDesc(DML_ACTIVATION_SCALED_ELU_OPERATOR_DESC& Desc, DmlUtil::FTensorDesc& TensorDesc)
+	{
+		Desc.InputTensor = &TensorDesc.Desc;
+		Desc.OutputTensor = &TensorDesc.Desc;
+		Desc.Alpha = Alpha;
+		Desc.Gamma = Gamma;
+	}
+
+	void InitDmlOpDesc(DML_ACTIVATION_ELU_OPERATOR_DESC& Desc, DmlUtil::FTensorDesc& TensorDesc)
+	{
+		Desc.InputTensor = &TensorDesc.Desc;
+		Desc.OutputTensor = &TensorDesc.Desc;
+		Desc.Alpha = Alpha;
+	}
+
+	void InitDmlOpDesc(DML_ACTIVATION_HARD_SIGMOID_OPERATOR_DESC& Desc, DmlUtil::FTensorDesc& TensorDesc)
+	{
+		Desc.InputTensor = &TensorDesc.Desc;
+		Desc.OutputTensor = &TensorDesc.Desc;
+		Desc.Alpha = Alpha;
+		Desc.Beta = Beta;
+	}
+
+	void InitDmlOpDesc(DML_ACTIVATION_LEAKY_RELU_OPERATOR_DESC& Desc, DmlUtil::FTensorDesc& TensorDesc)
+	{
+		Desc.InputTensor = &TensorDesc.Desc;
+		Desc.OutputTensor = &TensorDesc.Desc;
+		Desc.Alpha = Alpha;
+	}
+};
+
+template<> FOperatorDmlElementWiseUnary<DML_ACTIVATION_SCALED_ELU_OPERATOR_DESC, NNECore::Internal::EElementWiseUnaryOperatorType::Selu>::FOperatorDmlElementWiseUnary()
+	: Alpha(1.67326319217681884765625f), Beta(0.0f), Gamma(1.05070102214813232421875f), Num(1)
+{
+}
+
+template<> FOperatorDmlElementWiseUnary<DML_ACTIVATION_ELU_OPERATOR_DESC, NNECore::Internal::EElementWiseUnaryOperatorType::Elu>::FOperatorDmlElementWiseUnary()
+	: Alpha(1.0f), Beta(0.0f), Gamma(1.05070102214813232421875f), Num(1)
+{
+}
+
+template<> FOperatorDmlElementWiseUnary<DML_ACTIVATION_HARD_SIGMOID_OPERATOR_DESC, NNECore::Internal::EElementWiseUnaryOperatorType::HardSigmoid>::FOperatorDmlElementWiseUnary()
+	: Alpha(0.2f), Beta(0.5f), Gamma(0.0f), Num(1)
+{
+}
+
+template<> FOperatorDmlElementWiseUnary<DML_ACTIVATION_LEAKY_RELU_OPERATOR_DESC, NNECore::Internal::EElementWiseUnaryOperatorType::LeakyRelu>::FOperatorDmlElementWiseUnary()
+	: Alpha(0.01f), Beta(0.0f), Gamma(0.0f), Num(1)
+{
+}
+
+/**
+ * Element-wise binary ML operator implementation
+ */
+template
+<
+	typename TDmlElementWiseOpDescType,
+	NNECore::Internal::EElementWiseBinaryOperatorType OpType
+>
+class FOperatorDmlElementWiseBinary : public FOperatorDml
+{
+
+public:
+
+	static FOperatorDml* Create()
+	{
+		return new FOperatorDmlElementWiseBinary();
+	}
+
+private:
+
+	FOperatorDmlElementWiseBinary() = default;
+
+	uint32 Num { 1 };
+
+public:
+
+	//
+	//
+	//
+	virtual bool Initialize(IDMLDevice* Device, TArrayView<const NNECore::Internal::FTensor> InputTensors, TArrayView<const NNECore::Internal::FTensor> OutputTensors, const NNECore::FAttributeMap& Attributes) override
+	{
+		Num = OutputTensors[0].GetVolume();
+
+		const NNECore::Internal::FTensor& InputATensorDesc = InputTensors[0];
+		const NNECore::Internal::FTensor& InputBTensorDesc = InputTensors[1];
+		const NNECore::Internal::FTensor& OutputTensorDesc = OutputTensors[0];
+
+		// Initialize tensor descriptors
+		DmlUtil::FTensorDesc	DmlInputATensorDesc{};
+		DmlUtil::FTensorDesc	DmlInputBTensorDesc{};
+		DmlUtil::FTensorDesc	DmlOutputTensorDesc{};
+
+		if (!InitDmlTensorDesc(DmlInputATensorDesc, InputATensorDesc, OutputTensorDesc))
+		{
+			UE_LOG(LogNNE, Warning, TEXT("Failed to initialize tensor(s) for DML inference"));
+			return false;
+		}
+
+		if (!InitDmlTensorDesc(DmlInputBTensorDesc, InputBTensorDesc, OutputTensorDesc))
+		{
+			UE_LOG(LogNNE, Warning, TEXT("Failed to initialize tensor(s) for DML inference"));
+			return false;
+		}
+
+		if (!InitDmlTensorDesc(DmlOutputTensorDesc, OutputTensorDesc))
+		{
+			UE_LOG(LogNNE, Warning, TEXT("Failed to initialize tensor(s) for DML inference"));
+			return false;
+		}
+
+		TDmlElementWiseOpDescType	DmlElemWiseOpDesc{};
+
+		InitDmlOpDesc(DmlElemWiseOpDesc, DmlInputATensorDesc, DmlInputBTensorDesc, DmlOutputTensorDesc);
+
+		DML_OPERATOR_DESC DmlOpDesc{};
+
+		DmlOpDesc.Type = GetDmlElementWiseBinaryOpType<TDmlElementWiseOpDescType>();
+		DmlOpDesc.Desc = &DmlElemWiseOpDesc;
+
+		return CreateOperator(Device, DmlOpDesc);
+	}
+
+private:
+
+	template<typename OpDesc>
+	void InitDmlOpDesc(OpDesc& Desc, DmlUtil::FTensorDesc& LHSTensor, DmlUtil::FTensorDesc& RHSTensor, DmlUtil::FTensorDesc& OutputTensor)
+	{
+		Desc.ATensor = &LHSTensor.Desc;
+		Desc.BTensor = &RHSTensor.Desc;
+		Desc.OutputTensor = &OutputTensor.Desc;
+	}
+
+	void InitDmlOpDesc(DML_ELEMENT_WISE_POW_OPERATOR_DESC& Desc, DmlUtil::FTensorDesc& LHSTensor, DmlUtil::FTensorDesc& RHSTensor, DmlUtil::FTensorDesc& OutputTensor)
+	{
+		Desc.InputTensor = &LHSTensor.Desc;
+		Desc.ExponentTensor = &RHSTensor.Desc;
+		Desc.OutputTensor = &OutputTensor.Desc;
+	}
+
+	void InitDmlOpDesc(DML_ACTIVATION_PARAMETERIZED_RELU_OPERATOR_DESC& Desc, DmlUtil::FTensorDesc& LHSTensor, DmlUtil::FTensorDesc& RHSTensor, DmlUtil::FTensorDesc& OutputTensor)
+	{
+		Desc.InputTensor = &LHSTensor.Desc;
+		Desc.SlopeTensor = &RHSTensor.Desc;
+		Desc.OutputTensor = &OutputTensor.Desc;
+	}
+};
+
+void RegisterElementWiseUnaryOperators()
+{
+#define OP(DmlOpDesc, OpName) FOperatorRegistryDml::Get()->OpAdd(TEXT(#OpName), FOperatorDmlElementWiseUnary<DmlOpDesc, NNECore::Internal::EElementWiseUnaryOperatorType::OpName>::Create)
+
+	OP(DML_ELEMENT_WISE_ABS_OPERATOR_DESC, Abs);
+	OP(DML_ELEMENT_WISE_ACOS_OPERATOR_DESC, Acos);
+	OP(DML_ELEMENT_WISE_ACOSH_OPERATOR_DESC, Acosh);
+	OP(DML_ELEMENT_WISE_ASIN_OPERATOR_DESC, Asin);
+	OP(DML_ELEMENT_WISE_ASINH_OPERATOR_DESC, Asinh);
+	OP(DML_ELEMENT_WISE_ATAN_OPERATOR_DESC, Atan);
+	OP(DML_ELEMENT_WISE_ATANH_OPERATOR_DESC, Atanh);
+	OP(DML_ELEMENT_WISE_CEIL_OPERATOR_DESC, Ceil);
+	OP(DML_ELEMENT_WISE_COS_OPERATOR_DESC, Cos);
+	OP(DML_ELEMENT_WISE_COSH_OPERATOR_DESC, Cosh);
+	OP(DML_ACTIVATION_ELU_OPERATOR_DESC, Elu);
+	OP(DML_ELEMENT_WISE_ERF_OPERATOR_DESC, Erf);
+	OP(DML_ELEMENT_WISE_EXP_OPERATOR_DESC, Exp);
+	OP(DML_ELEMENT_WISE_FLOOR_OPERATOR_DESC, Floor);
+	OP(DML_ELEMENT_WISE_IS_INFINITY_OPERATOR_DESC, IsInf);
+	OP(DML_ELEMENT_WISE_IS_NAN_OPERATOR_DESC, IsNan);
+	OP(DML_ACTIVATION_HARD_SIGMOID_OPERATOR_DESC, HardSigmoid);
+	//OP(HardSwish);
+	OP(DML_ACTIVATION_LEAKY_RELU_OPERATOR_DESC, LeakyRelu);
+	OP(DML_ELEMENT_WISE_LOG_OPERATOR_DESC, Log);
+	OP(DML_ELEMENT_WISE_NEGATE_OPERATOR_DESC, Neg);
+	//OP(Not);
+	OP(DML_ELEMENT_WISE_RECIP_OPERATOR_DESC, Reciprocal);
+	OP(DML_ACTIVATION_RELU_OPERATOR_DESC, Relu);
+	OP(DML_ELEMENT_WISE_ROUND_OPERATOR_DESC, Round);
+	OP(DML_ACTIVATION_SCALED_ELU_OPERATOR_DESC, Selu);
+	OP(DML_ACTIVATION_SIGMOID_OPERATOR_DESC, Sigmoid);
+	OP(DML_ELEMENT_WISE_SIGN_OPERATOR_DESC, Sign);
+	OP(DML_ELEMENT_WISE_SIN_OPERATOR_DESC, Sin);
+	OP(DML_ELEMENT_WISE_SINH_OPERATOR_DESC, Sinh);
+	OP(DML_ACTIVATION_SOFTPLUS_OPERATOR_DESC, Softplus);
+	OP(DML_ACTIVATION_SOFTSIGN_OPERATOR_DESC, Softsign);
+	OP(DML_ELEMENT_WISE_SQRT_OPERATOR_DESC, Sqrt);
+	OP(DML_ELEMENT_WISE_TAN_OPERATOR_DESC, Tan);
+	OP(DML_ELEMENT_WISE_TANH_OPERATOR_DESC, Tanh);
+
+#undef OP
+}
+
+void RegisterElementWiseBinaryOperators()
+{
+#define OP(DmlOpDesc, OpName) FOperatorRegistryDml::Get()->OpAdd(TEXT(#OpName), FOperatorDmlElementWiseBinary<DmlOpDesc, NNECore::Internal::EElementWiseBinaryOperatorType::OpName>::Create)
+
+	OP(DML_ELEMENT_WISE_ADD_OPERATOR_DESC, Add);
+	// And
+	OP(DML_ELEMENT_WISE_DIVIDE_OPERATOR_DESC, Div);
+	// Equal
+	// Greater
+	// GreaterOrEqual
+	// Less
+	// LessOrEqual
+	// Mod
+	OP(DML_ELEMENT_WISE_MULTIPLY_OPERATOR_DESC, Mul);
+	// Or
+	OP(DML_ACTIVATION_PARAMETERIZED_RELU_OPERATOR_DESC, Prelu);
+	OP(DML_ELEMENT_WISE_POW_OPERATOR_DESC, Pow);
+	OP(DML_ELEMENT_WISE_SUBTRACT_OPERATOR_DESC, Sub);
+	// Xor
+
+#undef OP
+}
+
+struct FDmlOperatorElementWiseRegistrator
+{
+	FDmlOperatorElementWiseRegistrator()
+	{
+		RegisterElementWiseUnaryOperators();
+		RegisterElementWiseBinaryOperators();
+	}
+};
+
+static FDmlOperatorElementWiseRegistrator RegisterElementWiseOperators;
+
+} // namespace UE::NNERuntimeRDG::Private::Dml
+
+#endif // NNE_USE_DIRECTML
\ No newline at end of file
diff --git a/Engine/Plugins/Experimental/NNE/Source/NNERuntimeRDG/Private/Dml/NNEDmlOperatorGemm.cpp b/Engine/Plugins/Experimental/NNE/Source/NNERuntimeRDG/Private/Dml/NNEDmlOperatorGemm.cpp
new file mode 100644
--- /dev/null
+++ b/Engine/Plugins/Experimental/NNE/Source/NNERuntimeRDG/Private/Dml/NNEDmlOperatorGemm.cpp
@@ -0,0 +1,101 @@
+// Copyright Epic Games, Inc. All Rights Reserved.
+
+#ifdef NNE_USE_DIRECTML
+#include "NNEDmlOperator.h"
+
+namespace UE::NNERuntimeRDG::Private::Dml
+{
+
+/**
+ * Gemm
+ */
+class FOperatorDmlGemm : public FOperatorDml
+{
+public:
+
+	static FOperatorDml* Create()
+	{
+		return new FOperatorDmlGemm();
+	}
+
+	//
+	//
+	//
+	virtual bool Initialize(IDMLDevice* Device, TArrayView<const NNECore::Internal::FTensor> InputTensors, TArrayView<const NNECore::Internal::FTensor> OutputTensors, const NNECore::FAttributeMap& Attributes) override
+	{
+		// Setup attributes
+		float Alpha = 1.0f;
+		float Beta = 1.0f;
+		int32 TransA = 0;
+		int32 TransB = 0;
+		
+		Alpha = Attributes.GetValueOrDefault(TEXT("alpha"), Alpha);
+		Beta = Attributes.GetValueOrDefault(TEXT("beta"), Beta);
+		TransA = Attributes.GetValueOrDefault(TEXT("transA"), TransA);
+		TransB = Attributes.GetValueOrDefault(TEXT("transB"), TransB);
+		
+		const NNECore::Internal::FTensor& InputATensorDesc = InputTensors[0];
+		const NNECore::Internal::FTensor& InputBTensorDesc = InputTensors[1];
+		const NNECore::Internal::FTensor& OutputTensorDesc = OutputTensors[0];
+
+		// Initialize tensor descriptors
+		DmlUtil::FTensorDesc	DmlInputATensorDesc{};
+		DmlUtil::FTensorDesc	DmlInputBTensorDesc{};
+		DmlUtil::FTensorDesc	DmlInputCTensorDesc{};
+		DmlUtil::FTensorDesc	DmlOutputTensorDesc{};
+
+		if (!InitDmlTensorDesc(DmlInputATensorDesc, InputATensorDesc))
+		{
+			UE_LOG(LogNNE, Warning, TEXT("Failed to initialize tensor(s) for DML inference"));
+			return false;
+		}
+
+		if (!InitDmlTensorDesc(DmlInputBTensorDesc, InputBTensorDesc))
+		{
+			UE_LOG(LogNNE, Warning, TEXT("Failed to initialize tensor(s) for DML inference"));
+			return false;
+		}
+
+		if (InputTensors.Num() > 2)
+		{
+			const NNECore::Internal::FTensor& InputCTensorDesc = InputTensors[2];
+
+			if (!InitDmlTensorDesc(DmlInputCTensorDesc, InputCTensorDesc, OutputTensorDesc))
+			{
+				UE_LOG(LogNNE, Warning, TEXT("Failed to initialize tensor(s) for DML inference"));
+				return false;
+			}
+		}
+
+		if (!InitDmlTensorDesc(DmlOutputTensorDesc, OutputTensorDesc))
+		{
+			UE_LOG(LogNNE, Warning, TEXT("Failed to initialize tensor(s) for DML inference"));
+			return false;
+		}
+
+		DML_GEMM_OPERATOR_DESC	DmlGemmOpDesc{};
+
+		DmlGemmOpDesc.ATensor = &DmlInputATensorDesc.Desc;
+		DmlGemmOpDesc.BTensor = &DmlInputBTensorDesc.Desc;
+		DmlGemmOpDesc.CTensor = InputTensors.Num() > 2 ? &DmlInputCTensorDesc.Desc : nullptr;
+		DmlGemmOpDesc.OutputTensor = &DmlOutputTensorDesc.Desc;
+		DmlGemmOpDesc.Alpha = Alpha;
+		DmlGemmOpDesc.Beta = Beta;
+		DmlGemmOpDesc.TransA = TransA ? DML_MATRIX_TRANSFORM_TRANSPOSE : DML_MATRIX_TRANSFORM_NONE;
+		DmlGemmOpDesc.TransB = TransB ? DML_MATRIX_TRANSFORM_TRANSPOSE : DML_MATRIX_TRANSFORM_NONE;
+
+		DML_OPERATOR_DESC DmlOpDesc{};
+
+		DmlOpDesc.Type = DML_OPERATOR_GEMM;
+		DmlOpDesc.Desc = &DmlGemmOpDesc;
+
+		return CreateOperator(Device, DmlOpDesc);
+	}
+
+};
+
+NNE_DML_REGISTER_OP(Gemm)
+
+} // namespace UE::NNERuntimeRDG::Private::Dml
+
+#endif // NNE_USE_DIRECTML
diff --git a/Engine/Plugins/Experimental/NNE/Source/NNERuntimeRDG/Private/NNERuntimeRDGDml.cpp b/Engine/Plugins/Experimental/NNE/Source/NNERuntimeRDG/Private/NNERuntimeRDGDml.cpp
--- a/Engine/Plugins/Experimental/NNE/Source/NNERuntimeRDG/Private/NNERuntimeRDGDml.cpp
+++ b/Engine/Plugins/Experimental/NNE/Source/NNERuntimeRDG/Private/NNERuntimeRDGDml.cpp
@@ -3,2447 +3,22 @@
 #include "NNERuntimeRDGDml.h"
 
 #include "NNERuntimeRDG.h"
-#include "NNERuntimeRDGModel.h"
 #include "NNEUtilsModelOptimizer.h"
-#include "NNECoreOperator.h"
-#include "NNECoreAttributeMap.h"
 #include "NNECoreModelData.h"
 #include "NNECoreRuntimeFormat.h"
 #include "HAL/FileManager.h"
 #include "Interfaces/IPluginManager.h"
-#include "RenderGraphBuilder.h"
-#include "RenderGraphUtils.h"
 #include "Algo/Find.h"
 
-#ifdef NNE_USE_DIRECTML
-
-#include "Windows/AllowWindowsPlatformTypes.h"
-THIRD_PARTY_INCLUDES_START
-#include <unknwn.h>
-#include "Microsoft/COMPointer.h"
-#include "DirectML.h"
-THIRD_PARTY_INCLUDES_END
-#include "Windows/HideWindowsPlatformTypes.h"
-
-// DirectML is implemented using COM on all platforms
-#ifdef IID_GRAPHICS_PPV_ARGS
-#define DML_PPV_ARGS(x) __uuidof(*x), IID_PPV_ARGS_Helper(x)
-#else
-#define DML_PPV_ARGS(x) IID_PPV_ARGS(x)
-#endif
-
-#include "ID3D12DynamicRHI.h"
-
-#define NNE_USE_D3D12_RESOURCES
-
-namespace UE::NNERuntimeRDG::Private::Dml
-{
-/**
- * Utility function to get operator type
- */
-template<typename TElementWiseOpDesc>
-DML_OPERATOR_TYPE GetDmlElementWiseUnaryOpType();
-
-#define OP_EW(OpName) \
-template<> \
-DML_OPERATOR_TYPE GetDmlElementWiseUnaryOpType<DML_ELEMENT_WISE_##OpName##_OPERATOR_DESC>() { return DML_OPERATOR_ELEMENT_WISE_##OpName; }
-
-#define OP_AN(OpName) \
-template<> \
-DML_OPERATOR_TYPE GetDmlElementWiseUnaryOpType<DML_ACTIVATION_##OpName##_OPERATOR_DESC>() { return DML_OPERATOR_ACTIVATION_##OpName; }
-
-
-OP_EW(IDENTITY)
-OP_EW(ABS)
-OP_EW(ACOS)
-OP_EW(ACOSH)
-OP_EW(ASIN)
-OP_EW(ASINH)
-OP_EW(ATAN)
-OP_EW(ATANH)
-// BitShift
-// Cast
-OP_EW(CEIL)
-OP_EW(CLIP)
-OP_EW(COS)
-OP_EW(COSH)
-OP_AN(ELU)
-OP_EW(ERF)
-OP_EW(EXP)
-OP_EW(FLOOR)
-OP_EW(IS_INFINITY)
-OP_EW(IS_NAN)
-OP_AN(HARDMAX)
-OP_AN(HARD_SIGMOID)
-OP_AN(LEAKY_RELU)
-OP_EW(LOG)
-OP_EW(NEGATE)
-// Not
-OP_EW(RECIP)
-OP_AN(RELU)
-OP_EW(ROUND)
-OP_AN(SCALED_ELU)
-OP_AN(SIGMOID)
-OP_EW(SIGN)
-OP_EW(SIN)
-OP_EW(SINH)
-OP_AN(SOFTPLUS)
-OP_AN(SOFTSIGN)
-OP_EW(SQRT)
-OP_EW(TAN)
-OP_EW(TANH)
-
-#undef OP_EW
-#undef OP_AN
-
-
-/**
- * Utility function to get operator type
- */
-//template<typename TActivationOpDesc>
-//DML_OPERATOR_TYPE GetDmlActivationOpType();
-//
-//#define OP(OpName) \
-//template<> \
-//DML_OPERATOR_TYPE GetDmlActivationOpType<DML_ACTIVATION_##OpName##_OPERATOR_DESC>() { return DML_OPERATOR_ACTIVATION_##OpName; }
-//
-//OP(ELU)
-//OP(ERF)
-//OP(HARDMAX)
-//OP(HARD_SIGMOID)
-//OP(LEAKY_RELU)
-//OP(LINEAR)
-//OP(LOG_SOFTMAX)
-//OP(PARAMETERIZED_RELU)
-//OP(PARAMETRIC_SOFTPLUS)
-//OP(RELU)
-//OP(SCALED_ELU)
-//OP(SCALED_TANH)
-//OP(SOFTMAX)
-//OP(SOFTPLUS)
-//OP(SOFTSIGN)
-//OP(TANH)
-//OP(THRESHOLDED_RELU)
-//
-
-/**
- * Utility function to get operator type
- */
-template<typename TElementWiseOpDesc>
-DML_OPERATOR_TYPE GetDmlElementWiseBinaryOpType();
-
-#define OP_EW(OpName) \
-template<> \
-DML_OPERATOR_TYPE GetDmlElementWiseBinaryOpType<DML_ELEMENT_WISE_##OpName##_OPERATOR_DESC>() { return DML_OPERATOR_ELEMENT_WISE_##OpName; }
-
-#define OP_AN(OpName) \
-template<> \
-DML_OPERATOR_TYPE GetDmlElementWiseBinaryOpType<DML_ACTIVATION_##OpName##_OPERATOR_DESC>() { return DML_OPERATOR_ACTIVATION_##OpName; }
-
-OP_EW(ADD)
-//OP_EW(LOGICAL_AND)
-OP_EW(DIVIDE)
-//OP_EW(LOGICAL_EQUALS)
-//OP_EW(LOGICAL_GREATER_THAN)
-//OP_EW(LOGICAL_LESS_THAN)
-//OP_EW(MOD)
-OP_EW(MULTIPLY)
-//OP_EW(LOGICAL_OR)
-OP_AN(PARAMETERIZED_RELU)
-OP_EW(POW)
-OP_EW(SUBTRACT)
-//OP_EW(LOGICAL_XOR)
-
-#undef OP_EW
-#undef OP_AN
-
-namespace DmlUtil
-{
-	struct FTensorDesc
-	{
-		DML_BUFFER_TENSOR_DESC										BuffDesc;
-		DML_TENSOR_DESC												Desc;
-		TArray<uint32, TInlineAllocator<NNECore::FTensorShape::MaxRank>>	Sizes;
-		TArray<uint32, TInlineAllocator<NNECore::FTensorShape::MaxRank>>	Strides;
-	};
-
-	void SetTensorStrides(FTensorDesc& TensorDesc, const NNECore::Internal::FTensor& InputDesc)
-	{
-		uint32 CurrStride = 1;
-
-		TensorDesc.Strides.SetNum(InputDesc.GetShape().Rank());
-		
-		for (int32 i = InputDesc.GetShape().Rank() - 1; i >= 0; --i)
-		{
-			TensorDesc.Strides[i] = CurrStride;
-			CurrStride *= InputDesc.GetShape().GetData()[i];
-		}
-	}
-
-	void SetTensorSizesAndStridesForBroadcast(FTensorDesc& TensorDesc, const NNECore::Internal::FTensor& InputDesc, const NNECore::Internal::FTensor& TargetDesc)
-	{
-		static_assert(NNECore::FTensorShape::MaxRank <= 8);
-		
-		const uint32 TargetDimension = TargetDesc.GetShape().Rank() != -1 ? TargetDesc.GetShape().Rank() : InputDesc.GetShape().Rank();
-		checkf(TargetDesc.GetShape().Rank() >= InputDesc.GetShape().Rank(), TEXT("Can't broadcast tensor from rank %d to rank %d, should be inferior or equal."), InputDesc.GetShape().Rank(), TargetDimension);
-		
-		TensorDesc.Sizes.SetNum(TargetDimension);
-		TensorDesc.Strides.SetNum(TargetDimension);
-
-		const int32 DimensionOffset = int32(TargetDimension - InputDesc.GetShape().Rank());
-		
-		for (int32 i = 0; i < (int32) TargetDimension; ++i)
-		{
-			TensorDesc.Sizes[i] = i < DimensionOffset ? 1 : InputDesc.GetShape().GetData()[i - DimensionOffset];
-		}
-
-		uint32 CurrStride = 1;
-
-		for (int32 i = TargetDimension - 1; i >= 0; --i)
-		{
-			const bool bBroadcast = TensorDesc.Sizes[i] < TargetDesc.GetShape().GetData()[i];
-
-			TensorDesc.Strides[i] = bBroadcast ? 0 : CurrStride;
-			CurrStride *= TensorDesc.Sizes[i];
-
-			TensorDesc.Sizes[i] = TargetDesc.GetShape().GetData()[i];
-		}
-	}
-
-	inline bool IsSameShape(const NNECore::Internal::FTensor& Left, const NNECore::Internal::FTensor& Right)
-	{
-		if (Left.GetShape().Rank() != Right.GetShape().Rank())
-		{
-			return false;
-		}
-		
-		for (int32 Idx = 0; Idx < Left.GetShape().Rank(); ++Idx)
-		{
-			if (Left.GetShape().GetData()[Idx] != Right.GetShape().GetData()[Idx])
-			{
-				return false;
-			}
-		}
-
-		return true;
-	}
-
-	DML_TENSOR_DATA_TYPE GetTensorDataType(ENNETensorDataType DataType)
-	{
-		switch (DataType)
-		{
-		case ENNETensorDataType::Double:
-			return DML_TENSOR_DATA_TYPE_FLOAT64;
-
-		case ENNETensorDataType::Float:
-			return DML_TENSOR_DATA_TYPE_FLOAT32;
-
-		case ENNETensorDataType::Half:
-			return DML_TENSOR_DATA_TYPE_FLOAT16;
-
-		case ENNETensorDataType::UInt64:
-			return DML_TENSOR_DATA_TYPE_UINT64;
-
-		case ENNETensorDataType::UInt32:
-			return DML_TENSOR_DATA_TYPE_UINT32;
-
-		case ENNETensorDataType::UInt16:
-			return DML_TENSOR_DATA_TYPE_UINT16;
-
-		case ENNETensorDataType::UInt8:
-			return DML_TENSOR_DATA_TYPE_UINT8;
-
-		case ENNETensorDataType::Int64:
-			return DML_TENSOR_DATA_TYPE_INT64;
-
-		case ENNETensorDataType::Int32:
-			return DML_TENSOR_DATA_TYPE_INT32;
-
-		case ENNETensorDataType::Int16:
-			return DML_TENSOR_DATA_TYPE_INT16;
-
-		case ENNETensorDataType::Int8:
-			return DML_TENSOR_DATA_TYPE_INT8;
-
-		default:
-			return DML_TENSOR_DATA_TYPE_UNKNOWN;
-		}
-	}
-
-	inline uint64 CalculateBufferSize(const FTensorDesc& DmlTensor, const NNECore::Internal::FTensor& Desc)
-	{
-		uint64 ElemSizeInBytes = Desc.GetElemByteSize();
-		
-		if (ElemSizeInBytes == 0)
-		{
-			return 0;
-		}
-
-		uint64 MinSizeInBytes = 0;		
-		uint32 IndexOfLastElement = 0;
-
-		for (int32 i = 0; i < DmlTensor.Sizes.Num(); ++i)
-		{
-			IndexOfLastElement += (DmlTensor.Sizes[i] - 1) * DmlTensor.Strides[i];
-		}
-
-		MinSizeInBytes = (static_cast<uint64>(IndexOfLastElement) + 1) * ElemSizeInBytes;
-		
-		// Round up to the nearest 4 bytes.
-		MinSizeInBytes = (MinSizeInBytes + 3) & ~3ull;
-
-		return MinSizeInBytes;
-	}
-}
-
-//
-//
-//
-class FDeviceContextDml
-{
-public:
-
-	uint32							DeviceIndex;
-	ID3D12Device*					D3D12Device{ nullptr }; // Borrowed reference from RHI
-	TComPtr<IDMLDevice>				Device{ nullptr };
-	TComPtr<IDMLCommandRecorder>	CmdRec{ nullptr };
-};
-
-//
-// DirectML operator
-//
-class FOperatorDml
-{
-public:
-
-	virtual ~FOperatorDml() = default;
-
-	virtual bool Initialize(FDeviceContextDml* DevCtx, TArrayView<const NNECore::Internal::FTensor> InputTensors, TArrayView<const NNECore::Internal::FTensor> OutputTensors, const NNECore::FAttributeMap& Attributes) = 0;
-
-	IDMLOperator* GetOperator()
-	{
-		return DmlOp;
-	}
-
-protected:
-
-	bool InitDmlTensorDesc(DmlUtil::FTensorDesc& DmlTensorDesc, const NNECore::Internal::FTensor& TensorDesc)
-	{
-		DML_TENSOR_DATA_TYPE DmlDataType = DmlUtil::GetTensorDataType(TensorDesc.GetDataType());
-
-		if (DmlDataType == DML_TENSOR_DATA_TYPE_UNKNOWN)
-		{
-			DmlTensorDesc.BuffDesc = DML_BUFFER_TENSOR_DESC{};
-			DmlTensorDesc.Desc = DML_TENSOR_DESC{};
-
-			return false;
-		}
-
-		DmlTensorDesc.Sizes = TensorDesc.GetShape().GetData();
-		//Note: We should support tensor padding using strides defined in FTensorDesc
-		//DmlUtil::SetTensorStrides(DmlTensorDesc, TensorDesc.Strides);
-		
-		DML_BUFFER_TENSOR_DESC& BuffDesc = DmlTensorDesc.BuffDesc;
-
-		BuffDesc = DML_BUFFER_TENSOR_DESC{};
-		BuffDesc.DataType = DmlDataType;
-		BuffDesc.Flags = TensorDesc.HasPreparedData() ? DML_TENSOR_FLAG_OWNED_BY_DML : DML_TENSOR_FLAG_NONE;
-		BuffDesc.DimensionCount = TensorDesc.GetShape().Rank();
-		BuffDesc.Sizes = DmlTensorDesc.Sizes.GetData();
-		BuffDesc.Strides = nullptr;
-		BuffDesc.TotalTensorSizeInBytes = TensorDesc.GetDataSize();
-
-		DmlTensorDesc.Desc = DML_TENSOR_DESC{ DML_TENSOR_TYPE_BUFFER, &DmlTensorDesc.BuffDesc };
-
-		return true;
-	}
-
-	bool InitDmlTensorDesc(DmlUtil::FTensorDesc& DmlTensorDesc, const NNECore::Internal::FTensor& TensorDesc, const NNECore::Internal::FTensor& BroadcastDesc)
-	{
-		DML_TENSOR_DATA_TYPE DmlDataType = DmlUtil::GetTensorDataType(TensorDesc.GetDataType());
-
-		if (DmlDataType == DML_TENSOR_DATA_TYPE_UNKNOWN)
-		{
-			DmlTensorDesc.BuffDesc = DML_BUFFER_TENSOR_DESC{};
-			DmlTensorDesc.Desc = DML_TENSOR_DESC{};
-
-			return false;
-		}
-
-		if (DmlUtil::IsSameShape(TensorDesc, BroadcastDesc))
-		{
-			DmlTensorDesc.Sizes = TensorDesc.GetShape().GetData();
-			DmlUtil::SetTensorStrides(DmlTensorDesc, TensorDesc);
-		}
-		else if (TensorDesc.GetShape().Rank() > BroadcastDesc.GetShape().Rank())
-		{
-			return false;
-		}
-		else
-		{
-			DmlUtil::SetTensorSizesAndStridesForBroadcast(DmlTensorDesc, TensorDesc, BroadcastDesc);
-		}
-
-		//UE_LOG(LogNNE, Warning, TEXT("DmlTensorDesc:%d,%d,%d -> %d,%d,%d"),
-		//	TensorDesc.Sizes[0],
-		//	TensorDesc.Dimension > 1 ? TensorDesc.Sizes[1] : 0,
-		//	TensorDesc.Dimension > 2 ? TensorDesc.Sizes[2] : 0,
-		//	DmlTensorDesc.Sizes[0],
-		//	DmlTensorDesc.Dimension > 1 ? DmlTensorDesc.Sizes[1] : 0,
-		//	DmlTensorDesc.Dimension > 2 ? DmlTensorDesc.Sizes[2] : 0
-		//);
-
-		//UE_LOG(LogNNE, Warning, TEXT("DmlTensorStrides:%d,%d,%d"),
-		//	DmlTensorDesc.Strides[0], DmlTensorDesc.Strides[1], DmlTensorDesc.Strides[2]
-		//);
-
-		check(DmlTensorDesc.Strides.Num() == DmlTensorDesc.Sizes.Num());
-		
-		DML_BUFFER_TENSOR_DESC& BuffDesc = DmlTensorDesc.BuffDesc;
-		
-		BuffDesc = DML_BUFFER_TENSOR_DESC{};
-
-		BuffDesc.DataType = DmlDataType;
-		BuffDesc.Flags = TensorDesc.HasPreparedData() ? DML_TENSOR_FLAG_OWNED_BY_DML : DML_TENSOR_FLAG_NONE;
-		BuffDesc.DimensionCount = DmlTensorDesc.Sizes.Num();
-		BuffDesc.Sizes = DmlTensorDesc.Sizes.GetData();
-		BuffDesc.Strides = DmlTensorDesc.Strides.GetData();
-		BuffDesc.TotalTensorSizeInBytes = DmlUtil::CalculateBufferSize(DmlTensorDesc, TensorDesc);
-		
-		DmlTensorDesc.Desc = DML_TENSOR_DESC{ DML_TENSOR_TYPE_BUFFER, &DmlTensorDesc.BuffDesc };
-
-		return true;
-	}
-	
-	bool CreateOperator(const DML_OPERATOR_DESC& DmlOpDesc)
-	{
-		IDMLDevice* Device = DevCtx->Device;
-
-		// Create operator
-		IDMLOperator* Op = nullptr;
-
-		HRESULT Res;
-
-		Res = Device->CreateOperator(&DmlOpDesc, DML_PPV_ARGS(&Op));
-		if (!Op)
-		{
-			UE_LOG(LogNNE, Warning, TEXT("Error:Failed to create DML operator, hres:%d"), Res);
-			return false;
-		}
-
-		DmlOp = Op;
-
-		return DmlOp.IsValid();
-	}
-
-	FDeviceContextDml*			DevCtx;
-	TComPtr<IDMLOperator>		DmlOp;
-};
-
-/**
- * DirectML ML operator registry
- */
-using FOperatorRegistryDml = TOperatorRegistryRDG<FOperatorDml>;
-
-/**
- * Element-wise unary ML operator implementation
- */
-template
-<
-	typename DmlElementWiseOpDescType, 
-	NNECore::Internal::EElementWiseUnaryOperatorType OpType
->
-class FOperatorDmlElementWiseUnary : public FOperatorDml
-{
-public:
-
-	static FOperatorDml* Create()
-	{
-		return new FOperatorDmlElementWiseUnary();
-	}
-
-	virtual ~FOperatorDmlElementWiseUnary() = default;
-
-private:
-
-	FOperatorDmlElementWiseUnary() : Alpha(0.0f), Beta(0.0f), Gamma(0.0f), Num(1) {}
-	float Alpha;
-	float Beta;
-	float Gamma;
-	uint32 Num;
-
-public:
-
-	//
-	//
-	//
-	virtual bool Initialize(FDeviceContextDml* InDevCtx, TArrayView<const NNECore::Internal::FTensor> InputTensors, TArrayView<const NNECore::Internal::FTensor> OutputTensors, const NNECore::FAttributeMap& Attributes) override
-	{
-		Num = InputTensors[0].GetVolume();
-
-		DevCtx = InDevCtx;
-
-		const NNECore::Internal::FTensor& InputTensorDesc = InputTensors[0];
-		const NNECore::Internal::FTensor& OutputTensorDesc = OutputTensors[0];
-
-		Alpha = Attributes.GetValueOrDefault(TEXT("alpha"), Alpha);
-		Beta = Attributes.GetValueOrDefault(TEXT("beta"), Beta);
-		Gamma = Attributes.GetValueOrDefault(TEXT("gamma"), Gamma);
-
-		// Initialize tensor descriptor (it's same for both input and output)
-		DmlUtil::FTensorDesc	DmlTensorDesc{};
-
-		if (!InitDmlTensorDesc(DmlTensorDesc, InputTensorDesc))
-		{
-			UE_LOG(LogNNE, Warning, TEXT("Failed to initialize tensor(s) for DML inference"));
-			return false;
-		}
-
-		DmlElementWiseOpDescType	DmlElemWiseOpDesc{};
-
-		InitDmlOpDesc(DmlElemWiseOpDesc, DmlTensorDesc);
-
-		DML_OPERATOR_DESC DmlOpDesc{};
-
-		DmlOpDesc.Type = GetDmlElementWiseUnaryOpType<DmlElementWiseOpDescType>();
-		DmlOpDesc.Desc = &DmlElemWiseOpDesc;
-
-		return CreateOperator(DmlOpDesc);
-	}
-
-private:
-
-	template<typename OpDesc>
-	void InitDmlOpDesc(OpDesc& Desc, DmlUtil::FTensorDesc& TensorDesc)
-	{
-		Desc.InputTensor = &TensorDesc.Desc;
-		Desc.OutputTensor = &TensorDesc.Desc;
-	}
-
-	void InitDmlOpDesc(DML_ACTIVATION_SOFTPLUS_OPERATOR_DESC& Desc, DmlUtil::FTensorDesc& TensorDesc)
-	{
-		Desc.InputTensor = &TensorDesc.Desc;
-		Desc.OutputTensor = &TensorDesc.Desc;
-		Desc.Steepness = 1.0f;
-	}
-
-	void InitDmlOpDesc(DML_ACTIVATION_SCALED_ELU_OPERATOR_DESC& Desc, DmlUtil::FTensorDesc& TensorDesc)
-	{
-		Desc.InputTensor = &TensorDesc.Desc;
-		Desc.OutputTensor = &TensorDesc.Desc;
-		Desc.Alpha = Alpha;
-		Desc.Gamma = Gamma;
-	}
-
-	void InitDmlOpDesc(DML_ACTIVATION_ELU_OPERATOR_DESC& Desc, DmlUtil::FTensorDesc& TensorDesc)
-	{
-		Desc.InputTensor = &TensorDesc.Desc;
-		Desc.OutputTensor = &TensorDesc.Desc;
-		Desc.Alpha = Alpha;
-	}
-
-	void InitDmlOpDesc(DML_ACTIVATION_HARD_SIGMOID_OPERATOR_DESC& Desc, DmlUtil::FTensorDesc& TensorDesc)
-	{
-		Desc.InputTensor = &TensorDesc.Desc;
-		Desc.OutputTensor = &TensorDesc.Desc;
-		Desc.Alpha = Alpha;
-		Desc.Beta = Beta;
-	}
-
-	void InitDmlOpDesc(DML_ACTIVATION_LEAKY_RELU_OPERATOR_DESC& Desc, DmlUtil::FTensorDesc& TensorDesc)
-	{
-		Desc.InputTensor = &TensorDesc.Desc;
-		Desc.OutputTensor = &TensorDesc.Desc;
-		Desc.Alpha = Alpha;
-	}
-};
-
-template<> FOperatorDmlElementWiseUnary<DML_ACTIVATION_SCALED_ELU_OPERATOR_DESC, NNECore::Internal::EElementWiseUnaryOperatorType::Selu>::FOperatorDmlElementWiseUnary()
-	: Alpha(1.67326319217681884765625f), Beta(0.0f), Gamma(1.05070102214813232421875f), Num(1)
-{
-}
-
-template<> FOperatorDmlElementWiseUnary<DML_ACTIVATION_ELU_OPERATOR_DESC, NNECore::Internal::EElementWiseUnaryOperatorType::Elu>::FOperatorDmlElementWiseUnary()
-	: Alpha(1.0f), Beta(0.0f), Gamma(1.05070102214813232421875f), Num(1)
-{
-}
-
-template<> FOperatorDmlElementWiseUnary<DML_ACTIVATION_HARD_SIGMOID_OPERATOR_DESC, NNECore::Internal::EElementWiseUnaryOperatorType::HardSigmoid>::FOperatorDmlElementWiseUnary()
-	: Alpha(0.2f), Beta(0.5f), Gamma(0.0f), Num(1)
-{
-}
-
-template<> FOperatorDmlElementWiseUnary<DML_ACTIVATION_LEAKY_RELU_OPERATOR_DESC, NNECore::Internal::EElementWiseUnaryOperatorType::LeakyRelu>::FOperatorDmlElementWiseUnary()
-	: Alpha(0.01f), Beta(0.0f), Gamma(0.0f), Num(1)
-{
-}
-
-/**
- * Element-wise binary ML operator implementation
- */
-template
-<
-	typename TDmlElementWiseOpDescType,
-	NNECore::Internal::EElementWiseBinaryOperatorType OpType
->
-class FOperatorDmlElementWiseBinary : public FOperatorDml
-{
-
-public:
-
-	static FOperatorDml* Create()
-	{
-		return new FOperatorDmlElementWiseBinary();
-	}
-
-private:
-
-	FOperatorDmlElementWiseBinary() = default;
-
-	uint32 Num { 1 };
-
-public:
-
-	//
-	//
-	//
-	virtual bool Initialize(FDeviceContextDml* InDevCtx, TArrayView<const NNECore::Internal::FTensor> InputTensors, TArrayView<const NNECore::Internal::FTensor> OutputTensors, const NNECore::FAttributeMap& Attributes) override
-	{
-		Num = OutputTensors[0].GetVolume();
-
-		DevCtx = InDevCtx;
-
-		const NNECore::Internal::FTensor& InputATensorDesc = InputTensors[0];
-		const NNECore::Internal::FTensor& InputBTensorDesc = InputTensors[1];
-		const NNECore::Internal::FTensor& OutputTensorDesc = OutputTensors[0];
-
-		// Initialize tensor descriptors
-		DmlUtil::FTensorDesc	DmlInputATensorDesc{};
-		DmlUtil::FTensorDesc	DmlInputBTensorDesc{};
-		DmlUtil::FTensorDesc	DmlOutputTensorDesc{};
-
-		if (!InitDmlTensorDesc(DmlInputATensorDesc, InputATensorDesc, OutputTensorDesc))
-		{
-			UE_LOG(LogNNE, Warning, TEXT("Failed to initialize tensor(s) for DML inference"));
-			return false;
-		}
-
-		if (!InitDmlTensorDesc(DmlInputBTensorDesc, InputBTensorDesc, OutputTensorDesc))
-		{
-			UE_LOG(LogNNE, Warning, TEXT("Failed to initialize tensor(s) for DML inference"));
-			return false;
-		}
-
-		if (!InitDmlTensorDesc(DmlOutputTensorDesc, OutputTensorDesc))
-		{
-			UE_LOG(LogNNE, Warning, TEXT("Failed to initialize tensor(s) for DML inference"));
-			return false;
-		}
-
-		TDmlElementWiseOpDescType	DmlElemWiseOpDesc{};
-
-		InitDmlOpDesc(DmlElemWiseOpDesc, DmlInputATensorDesc, DmlInputBTensorDesc, DmlOutputTensorDesc);
-
-		DML_OPERATOR_DESC DmlOpDesc{};
-
-		DmlOpDesc.Type = GetDmlElementWiseBinaryOpType<TDmlElementWiseOpDescType>();
-		DmlOpDesc.Desc = &DmlElemWiseOpDesc;
-
-		return CreateOperator(DmlOpDesc);
-	}
-
-private:
-
-	template<typename OpDesc>
-	void InitDmlOpDesc(OpDesc& Desc, DmlUtil::FTensorDesc& LHSTensor, DmlUtil::FTensorDesc& RHSTensor, DmlUtil::FTensorDesc& OutputTensor)
-	{
-		Desc.ATensor = &LHSTensor.Desc;
-		Desc.BTensor = &RHSTensor.Desc;
-		Desc.OutputTensor = &OutputTensor.Desc;
-	}
-
-	void InitDmlOpDesc(DML_ELEMENT_WISE_POW_OPERATOR_DESC& Desc, DmlUtil::FTensorDesc& LHSTensor, DmlUtil::FTensorDesc& RHSTensor, DmlUtil::FTensorDesc& OutputTensor)
-	{
-		Desc.InputTensor = &LHSTensor.Desc;
-		Desc.ExponentTensor = &RHSTensor.Desc;
-		Desc.OutputTensor = &OutputTensor.Desc;
-	}
-
-	void InitDmlOpDesc(DML_ACTIVATION_PARAMETERIZED_RELU_OPERATOR_DESC& Desc, DmlUtil::FTensorDesc& LHSTensor, DmlUtil::FTensorDesc& RHSTensor, DmlUtil::FTensorDesc& OutputTensor)
-	{
-		Desc.InputTensor = &LHSTensor.Desc;
-		Desc.SlopeTensor = &RHSTensor.Desc;
-		Desc.OutputTensor = &OutputTensor.Desc;
-	}
-};
-
-/**
- * Gemm
- */
-class FOperatorDmlGemm : public FOperatorDml
-{
-public:
-
-	static FOperatorDml* Create()
-	{
-		return new FOperatorDmlGemm();
-	}
-
-	//
-	//
-	//
-	virtual bool Initialize(FDeviceContextDml* InDevCtx, TArrayView<const NNECore::Internal::FTensor> InputTensors, TArrayView<const NNECore::Internal::FTensor> OutputTensors, const NNECore::FAttributeMap& Attributes) override
-	{
-		// Setup attributes
-		float Alpha = 1.0f;
-		float Beta = 1.0f;
-		int32 TransA = 0;
-		int32 TransB = 0;
-		
-		Alpha = Attributes.GetValueOrDefault(TEXT("alpha"), Alpha);
-		Beta = Attributes.GetValueOrDefault(TEXT("beta"), Beta);
-		TransA = Attributes.GetValueOrDefault(TEXT("transA"), TransA);
-		TransB = Attributes.GetValueOrDefault(TEXT("transB"), TransB);
-		
-		DevCtx = InDevCtx;
-
-		const NNECore::Internal::FTensor& InputATensorDesc = InputTensors[0];
-		const NNECore::Internal::FTensor& InputBTensorDesc = InputTensors[1];
-		const NNECore::Internal::FTensor& OutputTensorDesc = OutputTensors[0];
-
-		// Initialize tensor descriptors
-		DmlUtil::FTensorDesc	DmlInputATensorDesc{};
-		DmlUtil::FTensorDesc	DmlInputBTensorDesc{};
-		DmlUtil::FTensorDesc	DmlInputCTensorDesc{};
-		DmlUtil::FTensorDesc	DmlOutputTensorDesc{};
-
-		if (!InitDmlTensorDesc(DmlInputATensorDesc, InputATensorDesc))
-		{
-			UE_LOG(LogNNE, Warning, TEXT("Failed to initialize tensor(s) for DML inference"));
-			return false;
-		}
-
-		if (!InitDmlTensorDesc(DmlInputBTensorDesc, InputBTensorDesc))
-		{
-			UE_LOG(LogNNE, Warning, TEXT("Failed to initialize tensor(s) for DML inference"));
-			return false;
-		}
-
-		if (InputTensors.Num() > 2)
-		{
-			const NNECore::Internal::FTensor& InputCTensorDesc = InputTensors[2];
-
-			if (!InitDmlTensorDesc(DmlInputCTensorDesc, InputCTensorDesc, OutputTensorDesc))
-			{
-				UE_LOG(LogNNE, Warning, TEXT("Failed to initialize tensor(s) for DML inference"));
-				return false;
-			}
-		}
-
-		if (!InitDmlTensorDesc(DmlOutputTensorDesc, OutputTensorDesc))
-		{
-			UE_LOG(LogNNE, Warning, TEXT("Failed to initialize tensor(s) for DML inference"));
-			return false;
-		}
-
-		DML_GEMM_OPERATOR_DESC	DmlGemmOpDesc{};
-
-		DmlGemmOpDesc.ATensor = &DmlInputATensorDesc.Desc;
-		DmlGemmOpDesc.BTensor = &DmlInputBTensorDesc.Desc;
-		DmlGemmOpDesc.CTensor = InputTensors.Num() > 2 ? &DmlInputCTensorDesc.Desc : nullptr;
-		DmlGemmOpDesc.OutputTensor = &DmlOutputTensorDesc.Desc;
-		DmlGemmOpDesc.Alpha = Alpha;
-		DmlGemmOpDesc.Beta = Beta;
-		DmlGemmOpDesc.TransA = TransA ? DML_MATRIX_TRANSFORM_TRANSPOSE : DML_MATRIX_TRANSFORM_NONE;
-		DmlGemmOpDesc.TransB = TransB ? DML_MATRIX_TRANSFORM_TRANSPOSE : DML_MATRIX_TRANSFORM_NONE;
-
-		DML_OPERATOR_DESC DmlOpDesc{};
-
-		DmlOpDesc.Type = DML_OPERATOR_GEMM;
-		DmlOpDesc.Desc = &DmlGemmOpDesc;
-
-		return CreateOperator(DmlOpDesc);
-	}
-
-};
-
-/**
- * Convolution
- */
-template <DML_CONVOLUTION_DIRECTION Direction>
-class FOperatorDmlConv : public FOperatorDml
-{
-	static constexpr uint32_t NcdhwDimensionCount = 5;
-	static constexpr uint32_t NcdhwSpatialDimensionCount = 3;
-	static constexpr uint32_t NonspatialDimensionCount = 2; // The batch and channel dimensions of NCW, NCHW, NCDHW....
-
-	using FSmallArray = TArray<uint32, TInlineAllocator<NcdhwSpatialDimensionCount>>;
-	using FIntArray = TArray<int32>;
-	using FUIntArray = TArray<uint32>;
-	
-	enum EAutoPad
-	{
-		NOTSET,
-		SAME_UPPER,
-		SAME_LOWER,
-		VALID
-	};
-
-	static EAutoPad AutoPadFromString(FStringView StringVal) 
-	{
-		if (StringVal.Compare(TEXT("NOTSET")) == 0) 
-		{
-			return EAutoPad::NOTSET;
-		}
-		else if (StringVal.Compare(TEXT("SAME_UPPER")) == 0)
-		{
-			return EAutoPad::SAME_UPPER;
-		}
-		else if (StringVal.Compare(TEXT("SAME_LOWER")) == 0)
-		{
-			return EAutoPad::SAME_LOWER;
-		}
-		else if (StringVal.Compare(TEXT("VALID")) == 0)
-		{
-			return EAutoPad::VALID;
-		}
-		else
-		{
-			return EAutoPad::NOTSET;
-		}
-	}
-
-	struct FConvArgs
-	{
-		EAutoPad		AutoPad;
-		FSmallArray		StartPadding;
-		FSmallArray		EndPadding;
-		FIntArray		OutPadding;
-		FIntArray		Dilations;
-		FIntArray		Strides;
-		FSmallArray		OutputShape;
-		uint32			NumDimensions;
-		FSmallArray		WindowSize;
-		int32			Group;
-
-		FConvArgs() = default;
-
-		//
-		//
-		//
-		bool Init(const NNECore::FTensorShape& InputShape, const NNECore::FTensorShape& FilterShape, const NNECore::FAttributeMap& Attributes)
-		{
-			check(InputShape.Rank() > NonspatialDimensionCount);
-			check(FilterShape.Rank() == InputShape.Rank());
-			
-			NumDimensions = InputShape.Rank() - NonspatialDimensionCount;
-
-			const FNNEAttributeValue*	AttrStrides = Attributes.GetAttributeValue(TEXT("strides"));
-
-			if (AttrStrides)
-			{
-				Strides = AttrStrides->GetValue<FIntArray>();
-			}
-			else
-			{
-				Strides.Init(1, NumDimensions);
-			}
-
-			check(Strides.Num() == 0 || Strides.Num() == FilterShape.Rank() - NonspatialDimensionCount);
-
-			const FNNEAttributeValue*	AttrDilations = Attributes.GetAttributeValue(TEXT("dilations"));
-
-			if (AttrDilations)
-			{
-				Dilations = AttrDilations->GetValue<FIntArray>();
-			}
-			else
-			{
-				Dilations.Init(1, NumDimensions);
-			}
-
-			check(Dilations.Num() == 0 || Dilations.Num() == FilterShape.Rank() - NonspatialDimensionCount);
-			
-			for (int32 Dim = FilterShape.Rank() - NumDimensions; Dim < FilterShape.Rank(); ++Dim)
-			{
-				WindowSize.Add(FilterShape.GetData()[Dim]);
-			}
-
-			AutoPad = AutoPadFromString(*Attributes.GetValue<FString>(TEXT("auto_pad")));
-
-			if (AutoPad == EAutoPad::NOTSET)
-			{
-				const FNNEAttributeValue* AttrPads = Attributes.GetAttributeValue(TEXT("pads"));
-				TArray<int32>	Pads;
-
-				if (AttrPads)
-				{
-					Pads = AttrPads->GetValue<TArray<int32>>();
-				}
-				else
-				{
-					Pads.Init(0, 2 * NumDimensions);
-				}
-
-				for (uint32 Dim = 0; Dim < NumDimensions; ++Dim)
-				{
-					StartPadding.Add(Pads[Dim]);
-					EndPadding.Add(Pads[Dim + NumDimensions]);
-				}
-			}
-			else if (AutoPad == EAutoPad::VALID)
-			{
-				StartPadding.Init(0, NumDimensions);
-				EndPadding.Init(0, NumDimensions);
-			}
-			else
-			{
-				StartPadding.Init(0, NumDimensions);
-				EndPadding.Init(0, NumDimensions);
-
-				SetAutoPadding(InputShape.GetData(), FilterShape.GetData());
-			}
-
-			Group = Attributes.GetValueOrDefault<int>(TEXT("group"), 1);
-
-			if (Direction == DML_CONVOLUTION_DIRECTION_FORWARD)
-			{
-				OutPadding.Init(0, NumDimensions);
-			}
-			else
-			{
-				const FNNEAttributeValue* AttrOutPadding = Attributes.GetAttributeValue(TEXT("output_padding"));
-
-				if (AttrOutPadding)
-				{
-					OutPadding = AttrOutPadding->GetValue<TArray<int32>>();
-				}
-				else
-				{
-					OutPadding.Init(0, NumDimensions);
-				}
-			}
-
-			if (Direction == DML_CONVOLUTION_DIRECTION_FORWARD)
-			{
-				SetOutputShape(InputShape.GetData(), FilterShape.GetData());
-			}
-			else
-			{
-				const FNNEAttributeValue* AttrOutShape = Attributes.GetAttributeValue(TEXT("output_shape"));
-
-				if (AttrOutShape)
-				{
-					TArray<int32> OutShapeVal = AttrOutShape->GetValue<TArray<int32>>();
-
-					for (int32 Value : OutShapeVal)
-					{
-						OutputShape.Add(uint32(Value));
-					}
-				}
-				else
-				{
-					SetOutputShape(InputShape.GetData(), FilterShape.GetData());
-				}
-			}
-
-			return true;
-		}
-
-		//
-		//
-		//
-		void SetAutoPadding(TConstArrayView<uint32> InputShape, TConstArrayView<uint32> FilterShape)
-		{
-			const uint32 DimOffset = NonspatialDimensionCount;
-
-			for (uint32 Dim = 0; Dim < NumDimensions; ++Dim)
-			{
-				uint32 InputLen = uint32(InputShape[Dim + DimOffset]);
-				uint32 StridedOutLen = (InputLen + Strides[Dim] - 1) / Strides[Dim];
-				uint32 KernelLen = 1 + (WindowSize[Dim] - 1) * Dilations[Dim];
-				uint32 Len = Strides[Dim] * (StridedOutLen - 1) + KernelLen;
-				uint32 Padding = (Len <= InputLen) ? 0 : (Len - InputLen);
-
-				if (AutoPad == EAutoPad::SAME_LOWER)
-				{
-					StartPadding[Dim] = (Padding + 1) / 2;
-				}
-				else
-				{
-					StartPadding[Dim] = Padding / 2;
-				}
-
-				EndPadding[Dim] = Padding - StartPadding[Dim];
-			}
-		}
-
-		//
-		//
-		//
-		void SetOutputShape(TConstArrayView<uint32> InputShape, TConstArrayView<uint32> FilterShape)
-		{
-			const uint32 DimOffset = NonspatialDimensionCount;
-
-			OutputShape.SetNumUninitialized(InputShape.Num());
-
-			if (Direction == DML_CONVOLUTION_DIRECTION_FORWARD)
-			{
-				OutputShape[0] = InputShape[0];
-				OutputShape[1] = FilterShape[0];
-
-				for (uint32 Dim = 0; Dim < NumDimensions; ++Dim)
-				{
-					uint32 InputLen = InputShape[Dim + DimOffset];
-					uint32 PaddedLen = InputLen + StartPadding[Dim] + EndPadding[Dim];
-					uint32 KernelLen = 1 + (WindowSize[Dim] - 1) * Dilations[Dim];
-
-					checkf(KernelLen <= PaddedLen, TEXT("KernelLen must < PaddedLen"));
-					checkf(Strides[Dim] != 0, TEXT("Strides must be != 0"));
-
-					uint32 StridableOutLen = PaddedLen - KernelLen;
-					uint32 OutLen = 1 + (StridableOutLen / Strides[Dim]);
-
-					OutputShape[Dim + DimOffset] = OutLen;
-				}
-			}
-			else
-			{
-				OutputShape[0] = InputShape[0];
-				OutputShape[1] = FilterShape[1] * Group;
-
-				for (uint32 Dim = 0; Dim < NumDimensions; ++Dim)
-				{
-					uint32 Padding = StartPadding[Dim] + EndPadding[Dim];
-					uint32 KernelLen = 1 + (WindowSize[Dim] - 1) * Dilations[Dim];
-
-					OutputShape[Dim + DimOffset] = (InputShape[Dim + DimOffset] - 1) * Strides[Dim] + KernelLen + OutPadding[Dim] - Padding;
-				}
-			}
-		}
-	};
-
-public:
-
-	static FOperatorDml* Create()
-	{
-		return new FOperatorDmlConv();
-	}
-
-	//
-	//
-	//
-	virtual bool Initialize(FDeviceContextDml* InDevCtx, TArrayView<const NNECore::Internal::FTensor> InputTensors, TArrayView<const NNECore::Internal::FTensor> OutputTensors, const NNECore::FAttributeMap& Attributes) override
-	{
-		DevCtx = InDevCtx;
-
-		const NNECore::Internal::FTensor& InputTensor = InputTensors[0];
-		const NNECore::Internal::FTensor& FilterTensor = InputTensors[1];
-		
-		FConvArgs	Args;
-		
-		if (!Args.Init(InputTensor.GetShape(), FilterTensor.GetShape(), Attributes))
-		{
-			return false;
-		}
-
-		NNECore::Internal::FTensor OutputTensor = OutputTensors[0];
-
-		if (Direction == DML_CONVOLUTION_DIRECTION_FORWARD)
-		{
-			OutputTensor.SetShape(NNECore::FTensorShape::Make(Args.OutputShape));
-		}
-
-		// Initialize tensor descriptors
-		DmlUtil::FTensorDesc	DmlInputTensor{};
-		DmlUtil::FTensorDesc	DmlFilterTensor{};
-		DmlUtil::FTensorDesc	DmlBiasTensor{};
-		DmlUtil::FTensorDesc	DmlOutputTensor{};
-
-		if (!InitDmlTensorDesc(DmlInputTensor, InputTensor))
-		{
-			UE_LOG(LogNNE, Warning, TEXT("Failed to initialize tensor(s) for DML inference"));
-			return false;
-		}
-
-		if (!InitDmlTensorDesc(DmlFilterTensor, FilterTensor))
-		{
-			UE_LOG(LogNNE, Warning, TEXT("Failed to initialize tensor(s) for DML inference"));
-			return false;
-		}
-
-		if (InputTensors.Num() > 2)
-		{
-			NNECore::Internal::FTensor BiasTensor = InputTensors[2];
-
-			if (BiasTensor.GetShape().Rank() < (int32) Args.NumDimensions)
-			{
-				FSmallArray Shape;
-
-				Shape.Add(1);
-				Shape.Add(BiasTensor.GetShape().GetData()[0]);
-				Shape.Add(1);
-
-				for (int32 Dim = 3; Dim < InputTensor.GetShape().Rank(); ++Dim)
-				{
-					Shape.Add(1);
-				}
-
-				BiasTensor.SetShape(NNECore::FTensorShape::Make(Shape));
-			}
-
-			if (!InitDmlTensorDesc(DmlBiasTensor, BiasTensor))
-			{
-				UE_LOG(LogNNE, Warning, TEXT("Failed to initialize tensor(s) for DML inference"));
-				return false;
-			}
-		}
-
-		if (!InitDmlTensorDesc(DmlOutputTensor, OutputTensor))
-		{
-			UE_LOG(LogNNE, Warning, TEXT("Failed to initialize tensor(s) for DML inference"));
-			return false;
-		}
-
-		DML_CONVOLUTION_OPERATOR_DESC	DmlConvOpDesc{};
-
-		DmlConvOpDesc.InputTensor = &DmlInputTensor.Desc;
-		DmlConvOpDesc.FilterTensor = &DmlFilterTensor.Desc;
-		DmlConvOpDesc.BiasTensor = InputTensors.Num() > 2 ? &DmlBiasTensor.Desc : nullptr;
-		DmlConvOpDesc.OutputTensor = &DmlOutputTensor.Desc;
-		DmlConvOpDesc.Mode = DML_CONVOLUTION_MODE_CROSS_CORRELATION;
-		DmlConvOpDesc.Direction = Direction;
-		DmlConvOpDesc.DimensionCount = Args.NumDimensions;
-		DmlConvOpDesc.Strides = (uint32*) Args.Strides.GetData();
-		DmlConvOpDesc.Dilations = (uint32*) Args.Dilations.GetData();
-		DmlConvOpDesc.StartPadding = Args.StartPadding.GetData();
-		DmlConvOpDesc.EndPadding = Args.EndPadding.GetData();
-		DmlConvOpDesc.OutputPadding = (uint32*) Args.OutPadding.GetData();
-		DmlConvOpDesc.GroupCount = Args.Group;
-
-		DML_OPERATOR_DESC DmlOpDesc{};
-
-		DmlOpDesc.Type = DML_OPERATOR_CONVOLUTION;
-		DmlOpDesc.Desc = &DmlConvOpDesc;
-
-		return CreateOperator(DmlOpDesc);
-	}
-};
-
-
-bool RegisterElementWiseUnaryOperators()
-{
-#define OP(DmlOpDesc, OpName) FOperatorRegistryDml::Get()->OpAdd(TEXT(#OpName), FOperatorDmlElementWiseUnary<DmlOpDesc, NNECore::Internal::EElementWiseUnaryOperatorType::OpName>::Create)
-
-	OP(DML_ELEMENT_WISE_ABS_OPERATOR_DESC, Abs);
-	OP(DML_ELEMENT_WISE_ACOS_OPERATOR_DESC, Acos);
-	OP(DML_ELEMENT_WISE_ACOSH_OPERATOR_DESC, Acosh);
-	OP(DML_ELEMENT_WISE_ASIN_OPERATOR_DESC, Asin);
-	OP(DML_ELEMENT_WISE_ASINH_OPERATOR_DESC, Asinh);
-	OP(DML_ELEMENT_WISE_ATAN_OPERATOR_DESC, Atan);
-	OP(DML_ELEMENT_WISE_ATANH_OPERATOR_DESC, Atanh);
-	OP(DML_ELEMENT_WISE_CEIL_OPERATOR_DESC, Ceil);
-	OP(DML_ELEMENT_WISE_COS_OPERATOR_DESC, Cos);
-	OP(DML_ELEMENT_WISE_COSH_OPERATOR_DESC, Cosh);
-	OP(DML_ACTIVATION_ELU_OPERATOR_DESC, Elu);
-	OP(DML_ELEMENT_WISE_ERF_OPERATOR_DESC, Erf);
-	OP(DML_ELEMENT_WISE_EXP_OPERATOR_DESC, Exp);
-	OP(DML_ELEMENT_WISE_FLOOR_OPERATOR_DESC, Floor);
-	OP(DML_ELEMENT_WISE_IS_INFINITY_OPERATOR_DESC, IsInf);
-	OP(DML_ELEMENT_WISE_IS_NAN_OPERATOR_DESC, IsNan);
-	OP(DML_ACTIVATION_HARD_SIGMOID_OPERATOR_DESC, HardSigmoid);
-	//OP(HardSwish);
-	OP(DML_ACTIVATION_LEAKY_RELU_OPERATOR_DESC, LeakyRelu);
-	OP(DML_ELEMENT_WISE_LOG_OPERATOR_DESC, Log);
-	OP(DML_ELEMENT_WISE_NEGATE_OPERATOR_DESC, Neg);
-	//OP(Not);
-	OP(DML_ELEMENT_WISE_RECIP_OPERATOR_DESC, Reciprocal);
-	OP(DML_ACTIVATION_RELU_OPERATOR_DESC, Relu);
-	OP(DML_ELEMENT_WISE_ROUND_OPERATOR_DESC, Round);
-	OP(DML_ACTIVATION_SCALED_ELU_OPERATOR_DESC, Selu);
-	OP(DML_ACTIVATION_SIGMOID_OPERATOR_DESC, Sigmoid);
-	OP(DML_ELEMENT_WISE_SIGN_OPERATOR_DESC, Sign);
-	OP(DML_ELEMENT_WISE_SIN_OPERATOR_DESC, Sin);
-	OP(DML_ELEMENT_WISE_SINH_OPERATOR_DESC, Sinh);
-	OP(DML_ACTIVATION_SOFTPLUS_OPERATOR_DESC, Softplus);
-	OP(DML_ACTIVATION_SOFTSIGN_OPERATOR_DESC, Softsign);
-	OP(DML_ELEMENT_WISE_SQRT_OPERATOR_DESC, Sqrt);
-	OP(DML_ELEMENT_WISE_TAN_OPERATOR_DESC, Tan);
-	OP(DML_ELEMENT_WISE_TANH_OPERATOR_DESC, Tanh);
-
-#undef OP
-
-	return true;
-}
-
-bool RegisterElementWiseBinaryOperators()
-{
-#define OP(DmlOpDesc, OpName) FOperatorRegistryDml::Get()->OpAdd(TEXT(#OpName), FOperatorDmlElementWiseBinary<DmlOpDesc, NNECore::Internal::EElementWiseBinaryOperatorType::OpName>::Create)
-
-	OP(DML_ELEMENT_WISE_ADD_OPERATOR_DESC, Add);
-	// And
-	OP(DML_ELEMENT_WISE_DIVIDE_OPERATOR_DESC, Div);
-	// Equal
-	// Greater
-	// GreaterOrEqual
-	// Less
-	// LessOrEqual
-	// Mod
-	OP(DML_ELEMENT_WISE_MULTIPLY_OPERATOR_DESC, Mul);
-	// Or
-	OP(DML_ACTIVATION_PARAMETERIZED_RELU_OPERATOR_DESC, Prelu);
-	OP(DML_ELEMENT_WISE_POW_OPERATOR_DESC, Pow);
-	OP(DML_ELEMENT_WISE_SUBTRACT_OPERATOR_DESC, Sub);
-	// Xor
-
-#undef OP
-
-	return true;
-}
-
-bool RegisterGemmOperator()
-{
-	FOperatorRegistryDml::Get()->OpAdd(TEXT("Gemm"), FOperatorDmlGemm::Create);
-	return true;
-}
-
-bool RegisterConvOperator()
-{
-	FOperatorRegistryDml::Get()->OpAdd(TEXT("Conv"), FOperatorDmlConv<DML_CONVOLUTION_DIRECTION_FORWARD>::Create);
-	return true;
-}
-
-bool RegisterConvTransposeOperator()
-{
-	FOperatorRegistryDml::Get()->OpAdd(TEXT("ConvTranspose"), FOperatorDmlConv<DML_CONVOLUTION_DIRECTION_BACKWARD>::Create);
-	return true;
-}
-
-//
-//
-//
-class FModel : public FModelRDG
-{
-	class FGraphBuilder;
-	class FBindingTable;
-	class FDebugName;
-
-public:
-
-	FModel();
-	~FModel();
-
-	bool Init(TConstArrayView<uint8> ModelData, FDeviceContextDml* InDevCtx);
-
-protected:
-
-	virtual void AddDispatchOps_RenderThread(FRDGBuilder& GraphBuilder) override;
-	virtual int PrepareTensorShapesAndData() override;
-
-private:
-
-	bool InitCompiledOp(TConstArrayView<int32> OpInputIndices, uint64 TensorDataSize);
-
-	FOperatorDml* OpCreate(const FString& Name, TArrayView<const NNECore::Internal::FTensor> InputTensorDesc, TArrayView<const NNECore::Internal::FTensor> OutputTensorDescs, const NNECore::FAttributeMap& Attributes);
-
-	FBufferRHIRef CreateRHIBuffer(FRHICommandListImmediate& RHICmdList, uint32 Size, EBufferUsageFlags Usage, ERHIAccess Access, const TCHAR* DbgName);
-	ID3D12Resource* CreateD3D12Buffer(uint32 Size, D3D12_RESOURCE_STATES ResourceState = D3D12_RESOURCE_STATE_COMMON, D3D12_HEAP_TYPE HeapType = D3D12_HEAP_TYPE_DEFAULT, const TCHAR* DebugName = nullptr);
-
-	//Note: This should go into RDG
-	static constexpr int32 MaxNumInputs = 32;
-	static constexpr int32 MaxNumOutputs = 4;
-
-	using FRHIBufferInputArray = TArray<FRHIBuffer*, TInlineAllocator<MaxNumInputs>>;
-	using FRHIBufferOutputArray = TArray<FRHIBuffer*, TInlineAllocator<MaxNumOutputs>>;
-
-	TComPtr<IDMLOperatorInitializer>	OpInit;
-	TComPtr<IDMLCompiledOperator>		CompiledOp;
-	FDeviceContextDml*					DevCtx;
-	TUniquePtr<FBindingTable>			BindingTable;
-	TComPtr<ID3D12DescriptorHeap>		DescHeap;
-	uint32								DescCount;
-	uint32								DescSize;
-
-	FRHIBufferInputArray				InputBuffers;
-	FRHIBufferOutputArray				OutputBuffers;
-#ifdef NNE_USE_D3D12_RESOURCES
-	TComPtr<ID3D12Resource>				PersistBuff;
-	TComPtr<ID3D12Resource>				TempBuff;
-#else
-	FBufferRHIRef						PersistBuff;
-#endif
-	uint64								MemSizeTemp;
-	uint64								MemSizePersist;
-	ID3D12DynamicRHI* DynamicRHI;
-};
-
-//
-//
-//
-class FModel::FDebugName
-{
-	static constexpr int32 Size = 128;
-
-public:
-
-	FDebugName()
-	{
-		Str[0] = '\0';
-		Length = 0;
-	}
-
-	FDebugName(const FString& InStr)
-	{
-		FTCHARToUTF8 Conv(*InStr);
-
-		Length = Conv.Length() + 1 < Size ? Conv.Length() + 1 : Size - 1;
-		FCStringAnsi::Strncpy(Str, Conv.Get(), Length);
-		Str[Length] = '\0';
-	}
-
-	FDebugName(FStringView InStr)
-	{
-		FTCHARToUTF8 Conv(InStr.GetData());
-
-		Length = Conv.Length() + 1 < Size ? Conv.Length() + 1 : Size - 1;
-		FCStringAnsi::Strncpy(Str, Conv.Get(), Length);
-		Str[Length] = '\0';
-	}
-	~FDebugName()
-	{
-	}
-
-	const char* Get() const
-	{
-		return Str;
-	}
-
-private:
-
-	char	Str[Size];
-	int32	Length;
-};
-
-//
-//
-//
-class FModel::FBindingTable
-{
-public:
-
-	bool Init(FModel* InModel)
-	{
-		Model = InModel;
-		DynamicRHI = InModel->DynamicRHI;
-
-		return true;
-	}
-
-#ifdef NNE_USE_D3D12_RESOURCES
-	void Bind(IDMLOperatorInitializer* OpInit, TConstArrayView<FRHIBuffer*> InputBuffers, ID3D12Resource* PersistBuff, ID3D12Resource* TempBuff = nullptr)
-#else
-	void Bind(IDMLOperatorInitializer* OpInit, TConstArrayView<FRHIBuffer*> InputBuffers, FRHIBuffer* PersistBuff, FRHIBuffer* TempBuff = nullptr)
-#endif
-	{
-		Reset(OpInit);
-
-		TArray<DML_BUFFER_BINDING, TInlineAllocator<MaxNumInputs>> Inputs;
-
-		for (FRHIBuffer* Buffer : InputBuffers)
-		{
-			if (Buffer)
-			{
-				Inputs.Emplace(MakeBind(Buffer));
-			}
-			else
-			{
-				Inputs.Add({});
-			}
-		}
-
-		DML_BUFFER_ARRAY_BINDING	InputBindArray{ Inputs.Num(), Inputs.GetData() };
-		DML_BINDING_DESC			InputBindArrayDesc{ DML_BINDING_TYPE_BUFFER_ARRAY, &InputBindArray };
-		
-		BindingTable->BindInputs(1, &InputBindArrayDesc);
-		
-		DML_BUFFER_BINDING			PersistBind {};
-		DML_BINDING_DESC			PersistBindDesc{ DML_BINDING_TYPE_BUFFER, &PersistBind };
-
-		if (PersistBuff)
-		{
-#ifdef NNE_USE_D3D12_RESOURCES
-			PersistBind = DML_BUFFER_BINDING { PersistBuff, 0, PersistBuff->GetDesc().Width };
-#else
-			PersistBind = MakeBind(PersistBuff);
-#endif
-		}
-
-		BindingTable->BindOutputs(1, &PersistBindDesc);
-
-		DML_BUFFER_BINDING			TempBind{};
-		DML_BINDING_DESC			TempBindDesc{ DML_BINDING_TYPE_BUFFER, &TempBind };
-
-		if (TempBuff)
-		{
-#ifdef NNE_USE_D3D12_RESOURCES
-			TempBind = { TempBuff, 0, TempBuff->GetDesc().Width };
-#else
-			TempBind = MakeBind(TempBuff);
-#endif
-			BindingTable->BindTemporaryResource(&TempBindDesc);
-		}
-	}
-
-#ifdef NNE_USE_D3D12_RESOURCES
-	void Bind(IDMLCompiledOperator* Op, TConstArrayView<FRHIBuffer*> InputBuffers, TConstArrayView<FRHIBuffer*> OutputBuffers, ID3D12Resource* PersistBuff = nullptr, ID3D12Resource* TempBuff = nullptr)
-#else
-	
-	void Bind(IDMLCompiledOperator* Op, TConstArrayView<FRHIBuffer*> InputBuffers, TConstArrayView<FRHIBuffer*> OutputBuffers, FRHIBuffer* PersistBuff = nullptr, FRHIBuffer* TempBuff = nullptr)
-#endif
-	{
-		Reset(Op);
-
-		for (FRHIBuffer* Buffer : InputBuffers)
-		{
-			AddBind(Buffer, InputBinds, InputBindDescs);
-		}
-
-		for (FRHIBuffer* Buffer : OutputBuffers)
-		{
-			AddBind(Buffer, OutputBinds, OutputBindDescs);
-		}
-
-		BindingTable->BindInputs(InputBinds.Num(), InputBindDescs.GetData());
-		BindingTable->BindOutputs(OutputBinds.Num(), OutputBindDescs.GetData());
-
-		DML_BUFFER_BINDING			PersistBind;
-		DML_BINDING_DESC			PersistBindDesc{ DML_BINDING_TYPE_BUFFER, &PersistBind };
-
-		if (PersistBuff)
-		{
-#ifdef NNE_USE_D3D12_RESOURCES
-			PersistBind = { PersistBuff, 0, PersistBuff->GetDesc().Width };
-#else
-			PersistBind =  MakeBind(PersistBuff);MakeBind(PersistBuff);
-#endif
-			BindingTable->BindPersistentResource(&PersistBindDesc);
-		}
-
-		DML_BUFFER_BINDING			TempBind{};
-		DML_BINDING_DESC			TempBindDesc{ DML_BINDING_TYPE_BUFFER, &TempBind };
-
-		if (TempBuff)
-		{
-			
-#ifdef NNE_USE_D3D12_RESOURCES
-			TempBind = { TempBuff, 0, TempBuff->GetDesc().Width };
-#else
-			TempBind = MakeBind(TempBuff);
-#endif
-			BindingTable->BindTemporaryResource(&TempBindDesc);
-		}
-	}
-
-
-	IDMLBindingTable* Get()
-	{
-		return BindingTable;
-	}
-
-private:
-
-	bool Reset(IDMLDispatchable* Disp)
-	{
-		InputBinds.Reset();
-		InputBindDescs.Reset();
-		OutputBinds.Reset();
-		OutputBindDescs.Reset();
-		
-		DML_BINDING_PROPERTIES BindingProps = Disp->GetBindingProperties();
-		DML_BINDING_TABLE_DESC Desc{};
-
-		Desc.Dispatchable = Disp;
-		Desc.CPUDescriptorHandle = CD3DX12_CPU_DESCRIPTOR_HANDLE(Model->DescHeap->GetCPUDescriptorHandleForHeapStart(), 0, Model->DescSize);
-		Desc.GPUDescriptorHandle = CD3DX12_GPU_DESCRIPTOR_HANDLE(Model->DescHeap->GetGPUDescriptorHandleForHeapStart(), 0, Model->DescSize);
-		Desc.SizeInDescriptors = Model->DescCount;
-
-		HRESULT Res;
-
-		if (!BindingTable)
-		{
-			Res = Model->DevCtx->Device->CreateBindingTable(&Desc, DML_PPV_ARGS(&BindingTable));
-			if (!BindingTable)
-			{
-				UE_LOG(LogNNE, Warning, TEXT("Failed to create DML binding table, res:%d"), Res);
-				return false;
-			}
-		}
-		else
-		{
-			BindingTable->Reset(&Desc);
-		}
-
-		return true;
-	}
-
-	template<class TBindingArray, class TDescArray>
-	void AddBind(FRHIBuffer* Buffer, TBindingArray& Bindings, TDescArray& Descs)
-	{
-		if (Buffer)
-		{
-			DML_BUFFER_BINDING& Bind = Bindings.Add_GetRef(MakeBind(Buffer));
-			Descs.Add({ DML_BINDING_TYPE_BUFFER, &Bind });
-		}
-		else
-		{
-			DML_BUFFER_BINDING& Bind = Bindings.Add_GetRef({});
-			Descs.Add({ DML_BINDING_TYPE_BUFFER, &Bind });
-		}
-	}
-
-	DML_BUFFER_BINDING MakeBind(FRHIBuffer* Buffer)
-	{
-		ID3D12Resource* Resource = DynamicRHI->RHIGetResource(Buffer);
-
-		return DML_BUFFER_BINDING{ Resource, 0, Buffer->GetSize() };
-	}
-
-	TComPtr<IDMLBindingTable>										BindingTable;
-	TArray<DML_BUFFER_BINDING, TInlineAllocator<MaxNumInputs>>		InputBinds;
-	TArray<DML_BINDING_DESC, TInlineAllocator<MaxNumInputs>>		InputBindDescs;
-	TArray<DML_BUFFER_BINDING, TInlineAllocator<MaxNumOutputs>>		OutputBinds;
-	TArray<DML_BINDING_DESC, TInlineAllocator<MaxNumOutputs>>		OutputBindDescs;
-	ID3D12DynamicRHI*												DynamicRHI;
-	FModel*															Model;
-};
-
-//
-//
-//
-class FModel::FGraphBuilder
-{
-public:
-
-	struct FOpDesc
-	{
-		FOperatorDml*	Op;
-		int32			InputStart;
-		int32			InputCount;
-		int32			OutputStart;
-		int32			OutputCount;
-		FDebugName		DbgName;
-	};
-
-	struct FGraphDesc
-	{
-		TConstArrayView<NNECore::Internal::FTensor>	AllTensors;
-		TConstArrayView<int32>		InputIndices;
-		TConstArrayView<int32>		OutputIndices;
-		TConstArrayView<int32>		WeightIndices;
-		TConstArrayView<int32>		IntermediateIndices;
-		TConstArrayView<FTensorRDG>	WeightTensors;
-		TConstArrayView<FOpDesc>	Operators;
-		TConstArrayView<int32>		OpInputIndices;
-		TConstArrayView<int32>		OpOutputIndices;
-	};
-	
-private:
-
-	enum class EEdgeType
-	{
-		Input,
-		Output,
-		Intermediate
-	};
-
-	struct FEdge
-	{
-		EEdgeType	Type;
-		int32		TensorIdx{ -1 };
-		int32		NodeSrc{ -1 };
-		int32		NodeSrcOutput{ -1 };
-		int32		NodeDst{ -1 };
-		int32		NodeDstInput{ -1 };
-		
-		FEdge(EEdgeType InType)
-			: Type(InType)
-		{
-		}
-
-		FEdge& SetTensorIdx(int32 Value)
-		{
-			TensorIdx = Value;
-			return *this;
-		}
-
-		FEdge& SetNodeSrc(int32 Value)
-		{
-			NodeSrc = Value;
-			return *this;
-		}
-
-		FEdge& SetNodeSrcOutput(int32 Value)
-		{
-			NodeSrcOutput = Value;
-			return *this;
-		}
-
-		FEdge& SetNodeDst(int32 Value)
-		{
-			NodeDst = Value;
-			return *this;
-		}
-
-		FEdge& SetNodeDstInput(int32 Value)
-		{
-			NodeDstInput = Value;
-			return *this;
-		}
-	};
-
-public:
-
-	IDMLCompiledOperator* Compile(FDeviceContextDml* DevCtx, const FGraphDesc& InGraph)
-	{
-		IDMLDevice*				Device = DevCtx->Device;
-		TComPtr<IDMLDevice1>	Device1;
-
-		Device1.FromQueryInterface(__uuidof(IDMLDevice1), DevCtx->Device);
-		check(Device1);
-		if (!Device1)
-		{
-			return nullptr;
-		}
-
-		if (!AddEdges(InGraph))
-		{
-			return nullptr;
-		}
-
-		TArray<DML_INPUT_GRAPH_EDGE_DESC>			InputEdges;
-		TArray<DML_OUTPUT_GRAPH_EDGE_DESC>			OutputEdges;
-		TArray<DML_INTERMEDIATE_GRAPH_EDGE_DESC>	IntermediateEdges;
-		TArray<FDebugName>							DbgInputNames;
-		TArray<FDebugName>							DbgIntermediateNames;
-		TArray<FDebugName>							DbgOutputNames;
-
-		DbgInputNames.Reserve(InGraph.AllTensors.Num());
-		DbgIntermediateNames.Reserve(InGraph.AllTensors.Num());
-
-		for (const FEdge& Edge : Edges)
-		{
-			if (Edge.Type == EEdgeType::Input)
-			{
-				check(Edge.NodeSrcOutput >= 0);
-				check(Edge.NodeDst >= 0);
-				check(Edge.NodeDstInput >= 0);
-
-				DML_INPUT_GRAPH_EDGE_DESC& Input = InputEdges.Add_GetRef({});
-
-				Input.GraphInputIndex = Edge.NodeSrcOutput;
-				Input.ToNodeIndex = Edge.NodeDst;
-				Input.ToNodeInputIndex = Edge.NodeDstInput;
-				
-				const FDebugName& DbgName = DbgInputNames.Add_GetRef(InGraph.AllTensors[Edge.TensorIdx].GetName());
-				Input.Name = DbgName.Get();
-			}
-			else if (Edge.Type == EEdgeType::Output)
-			{
-				check(Edge.NodeDstInput >= 0);
-				check(Edge.NodeSrc >= 0);
-				check(Edge.NodeSrcOutput >= 0);
-
-				DML_OUTPUT_GRAPH_EDGE_DESC&	Output = OutputEdges.Add_GetRef({});
-
-				Output.GraphOutputIndex = Edge.NodeDstInput;
-				Output.FromNodeIndex = Edge.NodeSrc;
-				Output.FromNodeOutputIndex = Edge.NodeSrcOutput;
-
-				const FDebugName& DbgName = DbgOutputNames.Add_GetRef(InGraph.AllTensors[Edge.TensorIdx].GetName());
-				Output.Name = DbgName.Get();
-			}
-			else if (Edge.Type == EEdgeType::Intermediate)
-			{
-				DML_INTERMEDIATE_GRAPH_EDGE_DESC& Intermediate = IntermediateEdges.Add_GetRef({});
-
-				check(Edge.NodeSrc >= 0);
-				check(Edge.NodeSrcOutput >= 0);
-				check(Edge.NodeDst >= 0);
-				check(Edge.NodeDstInput >= 0);
-
-				Intermediate.FromNodeIndex = Edge.NodeSrc;
-				Intermediate.FromNodeOutputIndex = Edge.NodeSrcOutput;
-				Intermediate.ToNodeIndex = Edge.NodeDst;
-				Intermediate.ToNodeInputIndex = Edge.NodeDstInput;
-
-				const FDebugName& DbgName = DbgIntermediateNames.Add_GetRef(InGraph.AllTensors[Edge.TensorIdx].GetName());
-				Intermediate.Name = DbgName.Get();
-			}
-		}
-			
-		TArray<DML_GRAPH_NODE_DESC>		Nodes;
-		TArray<DML_GRAPH_EDGE_DESC>		InputEdgeDescs;
-		TArray<DML_GRAPH_EDGE_DESC>		OutputEdgeDescs;
-		TArray<DML_GRAPH_EDGE_DESC>		IntermediateEdgeDescs;
-
-		Nodes.SetNumUninitialized(Operators.Num());
-
-		for (int32 Idx = 0; Idx < Operators.Num(); ++Idx)
-		{
-			DML_GRAPH_NODE_DESC& NodeDesc = Nodes[Idx];
-
-			NodeDesc.Type = DML_GRAPH_NODE_TYPE_OPERATOR;
-			NodeDesc.Desc = &Operators[Idx];
-		}
-
-		for (int32 Idx = 0; Idx < InputEdges.Num(); ++Idx)
-		{
-			DML_GRAPH_EDGE_DESC& Edge = InputEdgeDescs.Add_GetRef({});
-
-			Edge.Type = DML_GRAPH_EDGE_TYPE_INPUT;
-			Edge.Desc = &InputEdges[Idx];
-		}			
-
-		for (int32 Idx = 0; Idx < OutputEdges.Num(); ++Idx)
-		{
-			DML_GRAPH_EDGE_DESC& Edge = OutputEdgeDescs.Add_GetRef({});
-
-			Edge.Type = DML_GRAPH_EDGE_TYPE_OUTPUT;
-			Edge.Desc = &OutputEdges[Idx];				
-		}
-
-		for (int32 Idx = 0; Idx < IntermediateEdges.Num(); ++Idx)
-		{
-			DML_GRAPH_EDGE_DESC& Edge = IntermediateEdgeDescs.Add_GetRef({});
-
-			Edge.Type = DML_GRAPH_EDGE_TYPE_INTERMEDIATE;
-			Edge.Desc = &IntermediateEdges[Idx];
-		}
-			
-		DML_GRAPH_DESC	Graph = DML_GRAPH_DESC{};
-
-		Graph.InputCount = InputEdges.Num();
-		Graph.OutputCount = OutputEdges.Num();
-		Graph.NodeCount = Operators.Num();
-		Graph.Nodes = Nodes.GetData();
-		Graph.InputEdgeCount = InputEdgeDescs.Num();
-		Graph.InputEdges = InputEdgeDescs.GetData();
-		Graph.OutputEdgeCount = OutputEdgeDescs.Num();
-		Graph.OutputEdges = OutputEdgeDescs.GetData();
-		Graph.IntermediateEdgeCount = IntermediateEdgeDescs.Num();
-		Graph.IntermediateEdges = IntermediateEdgeDescs.GetData();
-
-		IDMLCompiledOperator* Op = nullptr;
-		HRESULT Res;
-			
-		Res = Device1->CompileGraph(&Graph, DML_EXECUTION_FLAG_NONE, DML_PPV_ARGS(&Op));
-		if (FAILED(Res))
-		{
-			UE_LOG(LogNNE, Warning, TEXT("Error:Failed to compile DML graph"));
-			Op = nullptr;
-		};
-
-		return Op;
-	}
-
-private:
-
-	bool AddEdges(const FGraphDesc& InGraph)
-	{
-		Edges.Reset();
-		Operators.Reset();
-		NumInputs = 0;
-		NumOutputs = 0;
-
-		for (int32 Idx = 0; Idx < InGraph.InputIndices.Num(); ++Idx)
-		{
-			const int32 TensorIdx = InGraph.InputIndices[Idx];
-			AddInput(TensorIdx);
-		}
-
-		for (int32 Idx = 0; Idx < InGraph.WeightIndices.Num(); ++Idx)
-		{
-			const int32 TensorIdx = InGraph.WeightIndices[Idx];
-			AddInput(TensorIdx);
-		}
-
-		//const int32 OutputIndexStart = InGraph.InputIndices.Num() + InGraph.WeightIndices.Num() + InGraph.IntermediateIndices.Num();
-
-		for (int32 Idx = 0; Idx < InGraph.OutputIndices.Num(); ++Idx)
-		{
-			const int32 TensorIdx = InGraph.OutputIndices[Idx];
-			AddOutput(TensorIdx);
-		}
-
-		Operators.Reset(InGraph.Operators.Num());
-		for (const FOpDesc& OpDesc : InGraph.Operators)
-		{
-			AddOp(OpDesc, InGraph);
-		}
-
-		//Note: Validate edges here
-		
-		return true;
-	}
-
-	void AddInput(int32 TensorIdx)
-	{
-		AddEdge(
-			FEdge(EEdgeType::Input)
-				.SetTensorIdx(TensorIdx)
-				.SetNodeSrcOutput(NumInputs)
-		);
-
-		++NumInputs;
-	}
-
-	void AddOutput(int32 TensorIdx)
-	{
-		AddEdge(
-			FEdge(EEdgeType::Output)
-				.SetTensorIdx(TensorIdx)
-				.SetNodeDstInput(NumOutputs)
-		);
-
-		++NumOutputs;
-	}
-
-	void AddIntermediate(int32 TensorIdx, int32 NodeSrc, int32 NodeSrcOutput)
-	{
-		FEdge* ConnEdge = Edges.FindByPredicate(
-				[TensorIdx](const FEdge& Curr)
-				{
-					return Curr.TensorIdx == TensorIdx;
-				}
-		);
-
-		if (ConnEdge)
-		{
-			ConnectEdgeSrc(TensorIdx, NodeSrc, NodeSrcOutput);
-		}
-		else
-		{
-			AddEdge(
-				FEdge(EEdgeType::Intermediate)
-					.SetTensorIdx(TensorIdx)
-					.SetNodeSrc(NodeSrc)
-					.SetNodeSrcOutput(NodeSrcOutput)
-			);
-		}
-	}
-
-	void AddEdge(const FEdge& Edge)
-	{
-		FEdge* StartEdge = Edges.FindByPredicate(
-				[Edge](const FEdge& Curr)
-				{
-					return Curr.TensorIdx == Edge.TensorIdx;
-				}
-		);
-
-		check(StartEdge == nullptr);
-		Edges.Add(Edge);
-	}
-
-	bool ConnectEdgeDst(int32 TensorIdx, int32 NodeDst, int32 NodeDstInput)
-	{
-		FEdge* StartEdge = Edges.FindByPredicate(
-				[TensorIdx](const FEdge& Curr)
-				{
-					return Curr.TensorIdx == TensorIdx;
-				}
-		);
-
-		bool bFoundEdge = false;
-
-		for (FEdge* Curr = StartEdge; Curr->TensorIdx == TensorIdx; ++Curr)
-		{
-			if (Curr->NodeDst == -1 && Curr->NodeDstInput == -1)
-			{
-				Curr->NodeDst = NodeDst;
-				Curr->NodeDstInput = NodeDstInput;
-				bFoundEdge = true;
-				break;
-			}
-			else if (Curr->NodeDst == NodeDst && Curr->NodeDstInput == NodeDstInput)
-			{
-				bFoundEdge = true;
-				break;
-			}
-		}
-
-		checkf(bFoundEdge, TEXT("ConnectEdgeDst() has failed"));
-		return bFoundEdge;
-	}
-
-	bool ConnectEdgeSrc(int32 TensorIdx, int32 NodeSrc, int32 NodeSrcOutput)
-	{
-		FEdge* StartEdge =
-			Edges.FindByPredicate(
-				[TensorIdx](const FEdge& Curr)
-				{
-					return Curr.TensorIdx == TensorIdx;
-				}
-		);
-
-		bool bFoundEdge = false;
-
-		for (FEdge* Curr = StartEdge; Curr->TensorIdx == TensorIdx; ++Curr)
-		{
-			if (Curr->NodeSrc == -1 && Curr->NodeSrcOutput == -1)
-			{
-				Curr->NodeSrc = NodeSrc;
-				Curr->NodeSrcOutput = NodeSrcOutput;
-				bFoundEdge = true;
-				break;
-			}
-			else if (Curr->NodeSrc == NodeSrc && Curr->NodeSrcOutput == NodeSrcOutput)
-			{
-				bFoundEdge = true;
-				break;
-			}
-		}
-
-		checkf(bFoundEdge, TEXT("ConnectEdgeSrc() has failed"));
-		return bFoundEdge;
-	}
-
-	void AddOp(const FOpDesc& InOp, const FGraphDesc& InGraph)
-	{
-		DML_OPERATOR_GRAPH_NODE_DESC& OpDesc = Operators.Add_GetRef({});
-
-		OpDesc.Operator = InOp.Op->GetOperator();
-		OpDesc.Name = InOp.DbgName.Get();
-
-		const int32 NodeIdx = Operators.Num() - 1;
-
-		for (int32 Idx = 0; Idx < InOp.InputCount; ++Idx)
-		{
-			const int32 TensorIdx = InGraph.OpInputIndices[Idx + InOp.InputStart];				
-			
-			ConnectEdgeDst(TensorIdx, NodeIdx, Idx);
-		}
-
-		for (int32 Idx = 0; Idx < InOp.OutputCount; ++Idx)
-		{
-			const int32 TensorIdx = InGraph.OpOutputIndices[Idx + InOp.OutputStart];
-
-			AddIntermediate(TensorIdx, NodeIdx, Idx);
-		}
-	}
-
-	TArray<FEdge>							Edges;
-	TArray<DML_OPERATOR_GRAPH_NODE_DESC>	Operators;
-	int32									NumInputs;
-	int32									NumOutputs;
-};
-
-//
-//
-//
-FModel::FModel()
-{
-}
-
-//
-//
-//
-FModel::~FModel()
-{
-}
-
-//
-//
-//
-bool FModel::Init(TConstArrayView<uint8> ModelData, FDeviceContextDml* InDevCtx)
-{
-	check(ModelData.Num() > 0);
-	FNNERuntimeFormat	Format;
-	int32 GuidSize = sizeof(UNNERuntimeRDGDmlImpl::GUID);
-	int32 VersionSize = sizeof(UNNERuntimeRDGDmlImpl::Version);
-
-	if (!LoadModel(ModelData, Format, GuidSize + VersionSize))
-	{
-		return false;
-	}
-
-	DevCtx = InDevCtx;
-	DynamicRHI = GetID3D12PlatformDynamicRHI();
-
-	HRESULT Res = DevCtx->Device->CreateOperatorInitializer(0, nullptr, DML_PPV_ARGS(&OpInit));
-	if (!OpInit)
-	{
-		UE_LOG(LogNNE, Warning, TEXT("Error:Failed to create DML operator initializer"));
-		return false;
-	}
-
-	// DirectML requires all tensors to be concrete
-	// Notes: to handle dynamic tensor desc, op should init from symbolic shapes
-	TArray<NNECore::Internal::FTensor>	Tensors;
-
-	Tensors.Reset(AllSymbolicTensorDescs.Num());
-	for (const NNECore::FTensorDesc& TensorDesc : AllSymbolicTensorDescs)
-	{
-		Tensors.Emplace(NNECore::Internal::FTensor::MakeFromSymbolicDesc(TensorDesc));
-	}
-
-	FGraphBuilder				DmlGraphBuilder;
-	FGraphBuilder::FGraphDesc	DmlGraphDesc;
-
-	DmlGraphDesc.AllTensors = Tensors;
-	DmlGraphDesc.InputIndices = InputTensorIndices;
-	DmlGraphDesc.OutputIndices = OutputTensorIndices;
-	DmlGraphDesc.WeightIndices = WeightTensorIndices;
-	DmlGraphDesc.IntermediateIndices = IntermediateTensorIndices;
-	DmlGraphDesc.WeightTensors = WeightTensorRDGs;
-	
-	TArray<FGraphBuilder::FOpDesc>	DmlGraphOperators;
-	TArray<int32>					OpInputIndices;
-	TArray<int32>					OpOutputIndices;
-	uint64							TensorDataSize = 0;
-
-	// Loop over all operators in the model and create them
-	for (int32 Idx = 0; Idx < Format.Operators.Num(); ++Idx)
-	{
-		const FString TypeName = Format.Operators[Idx].TypeName;
-
-		FGraphBuilder::FOpDesc		OpDesc;
-		TArray<NNECore::Internal::FTensor>				OpInputTensors;
-		TArray<NNECore::Internal::FTensor>				OpOutputTensors;
-		NNECore::FAttributeMap	AttributeMap;
-
-		OpDesc.InputStart = OpInputIndices.Num();
-		OpDesc.OutputStart = OpOutputIndices.Num();
-
-		for (int32 InputTensorIndex : Format.Operators[Idx].InTensors)
-		{
-			const int32 WeightTensorIdx = WeightTensorIndices.Find(InputTensorIndex);
-
-			if (WeightTensorIdx >= 0)
-			{
-				TConstArrayView<uint8> TensorData = OpInputTensors.Emplace_GetRef(WeightTensorRDGs[WeightTensorIdx]).GetPreparedData<uint8>();
-				
-				TensorDataSize += Align(TensorData.Num(), DML_MINIMUM_BUFFER_TENSOR_ALIGNMENT);
-			}
-			else
-			{
-				NNECore::FTensorDesc SymbolicTensorDesc = AllSymbolicTensorDescs[InputTensorIndex];
-				// Notes: to handle dynamic tensor desc, op should init from symbolic shapes
-				OpInputTensors.Emplace(NNECore::Internal::FTensor::MakeFromSymbolicDesc(SymbolicTensorDesc));
-			}
-
-			OpInputIndices.Emplace(InputTensorIndex);
-		}
-
-		for (int32 OutputTensorIndex : Format.Operators[Idx].OutTensors)
-		{
-			NNECore::FTensorDesc SymbolicTensorDesc = AllSymbolicTensorDescs[OutputTensorIndex];
-			// Notes: to handle dynamic tensor desc, op should init from symbolic shapes
-			OpOutputTensors.Emplace(NNECore::Internal::FTensor::MakeFromSymbolicDesc(SymbolicTensorDesc));
-			OpOutputIndices.Emplace(OutputTensorIndex);
-		}
-
-		for (const FNNEFormatAttributeDesc& Desc : Format.Operators[Idx].Attributes)
-		{
-			AttributeMap.SetAttribute(Desc.Name, Desc.Value);
-		}
-
-		OpDesc.Op = OpCreate(TypeName, OpInputTensors, OpOutputTensors, AttributeMap);
-
-		if (!OpDesc.Op)
-		{
-			UE_LOG(LogNNE, Warning, TEXT("Error:Failed to create operator:%s"), *TypeName);
-			return false;
-		}
-
-		OpDesc.InputCount = OpInputTensors.Num();
-		OpDesc.OutputCount = OpOutputTensors.Num();
-		OpDesc.DbgName = TypeName;
-
-		DmlGraphOperators.Emplace(OpDesc);
-	}
-
-	DmlGraphDesc.Operators = DmlGraphOperators;
-	DmlGraphDesc.OpInputIndices = OpInputIndices;
-	DmlGraphDesc.OpOutputIndices = OpOutputIndices;
-
-	CompiledOp = DmlGraphBuilder.Compile(DevCtx, DmlGraphDesc);
-	if (!CompiledOp)
-	{
-		return false;
-	}
-
-	return InitCompiledOp(OpInputIndices, TensorDataSize);
-}
-
-//
-//
-//
-bool FModel::InitCompiledOp(TConstArrayView<int32> OpInputIndices, uint64 TensorDataSize)
-{
-	static constexpr EBufferUsageFlags	WeightBuffUsage = BUF_UnorderedAccess;
-	static constexpr ERHIAccess			WeightBuffAccess = ERHIAccess::UAVMask;
-
-	static constexpr EBufferUsageFlags	PersistBuffFlags = BUF_Static | BUF_ShaderResource | BUF_UnorderedAccess;
-	static constexpr ERHIAccess			PersistBuffAccess = ERHIAccess::UAVMask;
-
-	static constexpr EBufferUsageFlags	TempBuffFlags = BUF_Volatile | BUF_UnorderedAccess;
-	static constexpr ERHIAccess			TempBuffAccess = ERHIAccess::UAVMask;
-
-	HRESULT					Res;
-	IDMLDevice*				Device = DevCtx->Device;
-	IDMLCompiledOperator*	CompiledOps[] = { CompiledOp };
-	
-	Res = OpInit->Reset(UE_ARRAY_COUNT(CompiledOps), CompiledOps);
-	if (FAILED(Res))
-	{
-		UE_LOG(LogNNE, Warning, TEXT("Error:Failed to reset DirectML operator initializer"));
-		return false;
-	}
-
-	DML_BINDING_PROPERTIES InitBindProps = OpInit->GetBindingProperties();
-	DML_BINDING_PROPERTIES ExecBindProps = CompiledOp->GetBindingProperties();
-
-	DescCount = std::max(InitBindProps.RequiredDescriptorCount, ExecBindProps.RequiredDescriptorCount);
-	
-	D3D12_DESCRIPTOR_HEAP_DESC	HeapDesc = {};
-
-	HeapDesc.Flags = D3D12_DESCRIPTOR_HEAP_FLAG_SHADER_VISIBLE;
-	HeapDesc.Type = D3D12_DESCRIPTOR_HEAP_TYPE_CBV_SRV_UAV;
-	HeapDesc.NumDescriptors = DescCount;
-
-	Res = DevCtx->D3D12Device->CreateDescriptorHeap(&HeapDesc, IID_PPV_ARGS(&DescHeap));
-	if (!DescHeap)
-	{
-		UE_LOG(LogNNE, Warning, TEXT("Failed to create descriptor heap, res:%x"), Res);
-		return false;
-	}
-
-	DescSize = DevCtx->D3D12Device->GetDescriptorHandleIncrementSize(D3D12_DESCRIPTOR_HEAP_TYPE_CBV_SRV_UAV);
-
-	BindingTable.Reset(new FBindingTable());
-	if (!BindingTable->Init(this))
-	{
-		return false;
-	}
-
-	MemSizeTemp = ExecBindProps.TemporaryResourceSize;
-	MemSizePersist = ExecBindProps.PersistentResourceSize;
-
-	FEvent* Signal = FGenericPlatformProcess::GetSynchEventFromPool(false);
-
-	ENQUEUE_RENDER_COMMAND(FModel_SetTensorData)
-	(
-		[
-			this, 
-			Signal, 
-			InitTempMemSize = InitBindProps.TemporaryResourceSize,
-			TensorDataSize
-		]
-		(FRHICommandListImmediate& RHICmdList)
-		{
-			FRHIBufferInputArray	Inputs;
-
-			for (int32 InputIdx : InputTensorIndices)
-			{
-				Inputs.Emplace(nullptr);
-			}
-
-			TArray<CD3DX12_RESOURCE_BARRIER, TInlineAllocator<MaxNumInputs>>	Barriers;
-			FGPUFenceRHIRef	UploadFence = nullptr;
-
-			if (TensorDataSize)
-			{
-				UploadFence = RHICreateGPUFence(TEXT("FInferenceModel_UploadFence"));
-				
-				FRHIBuffer*		UploadBuff = CreateRHIBuffer(RHICmdList, TensorDataSize, BUF_ShaderResource | BUF_Dynamic | BUF_FastVRAM, ERHIAccess::CopySrc, TEXT("FInferenceModel_UploadBuffer"));
-				uint8*			UploadBuffPtr = static_cast<uint8*>(RHICmdList.LockBuffer(UploadBuff, 0, TensorDataSize, RLM_WriteOnly_NoOverwrite));
-				uint64			UploadOffset = 0;
-
-				for (const FTensorRDG& Tensor : WeightTensorRDGs)
-				{
-					TConstArrayView<uint8>	TensorData = Tensor.GetPreparedData<uint8>();
-
-					FBufferRHIRef WeightBuff;
-
-					WeightBuff = CreateRHIBuffer(RHICmdList, TensorData.Num(), WeightBuffUsage, WeightBuffAccess, TEXT("FModel_TensorWeights"));
-				
-					FMemory::Memcpy(UploadBuffPtr + UploadOffset, TensorData.GetData(), TensorData.Num());
-					RHICmdList.CopyBufferRegion(WeightBuff, 0, UploadBuff, UploadOffset, TensorData.Num());
-					UploadOffset += Align(TensorData.Num(), DML_MINIMUM_BUFFER_TENSOR_ALIGNMENT);
-
-					Inputs.Emplace(WeightBuff);
-				
-					Barriers.Emplace(
-						CD3DX12_RESOURCE_BARRIER::Transition(
-							DynamicRHI->RHIGetResource(WeightBuff),
-							D3D12_RESOURCE_STATE_COPY_DEST,
-							D3D12_RESOURCE_STATE_UNORDERED_ACCESS)
-					);
-				}
-
-				RHICmdList.UnlockBuffer(UploadBuff);
-				RHICmdList.WriteGPUFence(UploadFence);
-			}
-
-			if (MemSizePersist)
-			{
-#ifdef NNE_USE_D3D12_RESOURCES
-				PersistBuff = CreateD3D12Buffer(MemSizePersist);
-#else		
-				PersistBuff = CreateRHIBuffer(RHICmdList, MemSizePersist, PersistBuffFlags, PersistBuffAccess, TEXT("FModelDml_PeristBuff"));
-#endif
-			}
-
-			if (MemSizeTemp)
-			{
-#ifdef NNE_USE_D3D12_RESOURCES
-				TempBuff = CreateD3D12Buffer(MemSizeTemp);
-#else
-				TempBuff = CreateRHIBuffer(RHICmdList, MemSizeTemp, TempBuffFlags, TempBuffAccess, TEXT("FModelDml_TempBuff"));
-#endif
-			}
-
-#ifdef NNE_USE_D3D12_RESOURCES
-			ID3D12Resource* InitTempBuff = nullptr;
-#else
-			FBufferRHIRef InitTempBuff;
-#endif
-
-			if (InitTempMemSize)
-			{
-#ifdef NNE_USE_D3D12_RESOURCES
-				InitTempBuff = CreateD3D12Buffer(InitTempMemSize);
-#else
-				TempBuff = CreateRHIBuffer(RHICmdList, InitTempMemSize, TempBuffFlags, TempBuffAccess, TEXT("FModelDml_InitTempBuff"));
-#endif
-			}
-
-			RHICmdList.EnqueueLambda(
-				[this, Inputs, Barriers, InitTempBuff, UploadFence](FRHICommandListImmediate& RHICmdList)
-				{
-					while (UploadFence && UploadFence->NumPendingWriteCommands.GetValue() > 0)
-					{
-						FPlatformProcess::Sleep(0.001);
-					}
-
-					ID3D12GraphicsCommandList* D3DCmdList = nullptr;
-
-					D3DCmdList = DynamicRHI->RHIGetGraphicsCommandList(DevCtx->DeviceIndex);
-
-					BindingTable->Bind(OpInit, Inputs, PersistBuff, InitTempBuff);
-			
-					D3DCmdList = DynamicRHI->RHIGetGraphicsCommandList(DevCtx->DeviceIndex);
-
-					if (!Barriers.IsEmpty())
-					{
-						D3DCmdList->ResourceBarrier(Barriers.Num(), Barriers.GetData());
-					}
-					D3DCmdList->SetDescriptorHeaps(1, &DescHeap);
-					DevCtx->CmdRec->RecordDispatch(D3DCmdList, OpInit, BindingTable->Get());
-
-					DynamicRHI->RHIFinishExternalComputeWork(DevCtx->DeviceIndex, D3DCmdList);
-				}
-			);			
-			
-			RHICmdList.ImmediateFlush(EImmediateFlushType::FlushRHIThread);
-			Signal->Trigger();
-		}
-	);
-
-	Signal->Wait();
-	FGenericPlatformProcess::ReturnSynchEventToPool(Signal);
-
-	return true;
-}
-
-BEGIN_SHADER_PARAMETER_STRUCT(FTensorBufferParamsDml, )
-	RDG_BUFFER_ACCESS(Buffer, ERHIAccess::UAVCompute)
-END_SHADER_PARAMETER_STRUCT()
-
-//
-//
-//
-void FModel::AddDispatchOps_RenderThread(FRDGBuilder& GraphBuilder)
-{
-	const ERDGPassFlags TransitionBuffFlags = ERDGPassFlags::Compute | ERDGPassFlags::NeverCull;
-
-	InputBuffers.SetNumUninitialized(InputTensorIndices.Num() + WeightTensorIndices.Num());
-	OutputBuffers.SetNumUninitialized(OutputTensorIndices.Num());
-
-	for (int32 Idx = 0; Idx < InputTensorIndices.Num(); ++Idx)
-	{
-		FTensorBufferParamsDml* Params = GraphBuilder.AllocParameters<FTensorBufferParamsDml>();
-		Params->Buffer = AllTensorRDGRefs[InputTensorIndices[Idx]]->GetBuffer();
-
-		GraphBuilder.AddPass(
-			RDG_EVENT_NAME("FModel_Dispatch_GetInputBuffer"),
-			Params,
-			TransitionBuffFlags,
-			[this, Idx, Params](FRHICommandListImmediate& RHICmdList)
-			{
-				InputBuffers[Idx] = Params->Buffer->GetRHI();
-			}
-		);
-	}
-
-	for (int32 Idx = 0; Idx < WeightTensorIndices.Num(); ++Idx)
-	{
-		InputBuffers[Idx + InputTensorIndices.Num()] = nullptr;
-	}
-
-	for (int32 Idx = 0; Idx < OutputTensorIndices.Num(); ++Idx)
-	{
-		FTensorBufferParamsDml* Params = GraphBuilder.AllocParameters<FTensorBufferParamsDml>();
-		Params->Buffer = AllTensorRDGRefs[OutputTensorIndices[Idx]]->GetBuffer();
-
-		GraphBuilder.AddPass(
-			RDG_EVENT_NAME("FModel_Dispatch_GetOutputBuffer"),
-			Params,
-			TransitionBuffFlags,
-			[this, Idx, Params](FRHICommandListImmediate& RHICmdList)
-			{
-				OutputBuffers[Idx] = Params->Buffer->GetRHI();
-			}
-		);
-	}
-	
-	GraphBuilder.AddPass(
-		RDG_EVENT_NAME("FModel_Dispatch"),
-		ERDGPassFlags::None | ERDGPassFlags::NeverCull,
-		[this](FRHICommandListImmediate& RHICmdList)
-		{
-			RHICmdList.EnqueueLambda(
-				[this](FRHICommandListImmediate& RHICmdList)
-				{
-					TArray<CD3DX12_RESOURCE_BARRIER, TInlineAllocator<MaxNumInputs + MaxNumOutputs>>	Barriers;
-
-					for (FRHIBuffer* Buffer : InputBuffers)
-					{
-						if (!Buffer)
-						{
-							continue;
-						}
-
-						ID3D12Resource* Resource = DynamicRHI->RHIGetResource(Buffer);
-
-						//Note: We should not assume COPY_DEST state
-						Barriers.Emplace(
-							CD3DX12_RESOURCE_BARRIER::Transition(
-								Resource,
-								D3D12_RESOURCE_STATE_COPY_DEST,
-								D3D12_RESOURCE_STATE_UNORDERED_ACCESS)
-						);
-					}
-
-					//Note: We should use this instead of NNE_USE_D3D12_RESOURCES
-					//FBufferRHIRef TempBuff = MemSizeTemp ? CreateRHIBuffer(RHICmdList, MemSizeTemp, TempBuffUsage, TempBuffAccess, TEXT("FModel_Dispatch_TempBuff")) : nullptr;
-
-					BindingTable->Bind(CompiledOp, InputBuffers, OutputBuffers, PersistBuff, TempBuff);
-
-					ID3D12GraphicsCommandList* D3DCmdList = nullptr;
-
-					D3DCmdList = DynamicRHI->RHIGetGraphicsCommandList(DevCtx->DeviceIndex);
-					D3DCmdList->SetDescriptorHeaps(1, &DescHeap);
-					D3DCmdList->ResourceBarrier(Barriers.Num(), Barriers.GetData());
-					DevCtx->CmdRec->RecordDispatch(D3DCmdList, CompiledOp, BindingTable->Get());
-
-					DynamicRHI->RHIFinishExternalComputeWork(DevCtx->DeviceIndex, D3DCmdList);
-				}
-			);
-		}
-	);
-}
-
-//
-// Create operator
-//
-FOperatorDml* FModel::OpCreate(const FString& OpName, TArrayView<const NNECore::Internal::FTensor> InputTensorDescs, TArrayView<const NNECore::Internal::FTensor> OutputTensorDescs, const NNECore::FAttributeMap& Attributes)
-{
-	FOperatorRegistryDml::OperatorCreateFunc CreateFn = FOperatorRegistryDml::Get()->OpFind(OpName);
-
-	if (!CreateFn)
-	{
-		UE_LOG(LogNNE, Warning, TEXT("Dml MLOperatorRegistry failed to find operator:%s"), *OpName);
-		return nullptr;
-	}
-
-	FOperatorDml* Op = CreateFn();
-
-	if (!Op->Initialize(DevCtx, InputTensorDescs, OutputTensorDescs, Attributes))
-	{
-		delete Op;
-
-		UE_LOG(LogNNE, Warning, TEXT("Error:Failed to initialize operator:%s"), *OpName);
-		return nullptr;
-	}
-
-	Op->GetOperator()->SetName(*OpName);
-
-	return Op;
-}
-
-FBufferRHIRef FModel::CreateRHIBuffer(FRHICommandListImmediate& RHICmdList, uint32 Size, EBufferUsageFlags Usage, ERHIAccess Access, const TCHAR* DbgName)
-{
-	FBufferRHIRef Buff = nullptr;
-
-	if (Size)
-	{
-		FRHIResourceCreateInfo CreateInfo(DbgName);
-		
-		Buff = RHICmdList.CreateBuffer(Size, Usage, 1, Access, CreateInfo);
-	}
+#include "Dml/NNEDmlModel.h"
 
-	check(Buff);
-	return Buff;
-}
-
-ID3D12Resource* FModel::CreateD3D12Buffer(uint32 Size, D3D12_RESOURCE_STATES ResourceState, D3D12_HEAP_TYPE HeapType, const TCHAR* DebugName)
-{
-	ID3D12Resource* Resource = nullptr;
-
-	D3D12_RESOURCE_DESC ResourceDesc = CD3DX12_RESOURCE_DESC::Buffer(Size,D3D12_RESOURCE_FLAG_ALLOW_UNORDERED_ACCESS);
-	CD3DX12_HEAP_PROPERTIES HeapProps = CD3DX12_HEAP_PROPERTIES(HeapType);
-	HRESULT Res;
-
-	Res = DevCtx->D3D12Device->CreateCommittedResource(
-		&HeapProps,
-		D3D12_HEAP_FLAG_NONE,
-		&ResourceDesc,
-		ResourceState,
-		nullptr,
-		IID_PPV_ARGS(&Resource));
-
-	if (FAILED(Res))
-	{
-		UE_LOG(LogNNE, Warning, TEXT("Error:FInferenceModel failed to create D3D12 resource"));
-		return nullptr;
-	}
-
-	if (Resource && DebugName)
-	{
-		Resource->SetName(DebugName);
-	}
+#ifdef NNE_USE_DIRECTML
 
-	return Resource;
-}
+#include "Dml/NNEDmlOperator.h"
+#include "ID3D12DynamicRHI.h"
 
-//
-//
-//
-int FModel::PrepareTensorShapesAndData()
+namespace UE::NNERuntimeRDG::Private::Dml
 {
-	for (NNECore::FTensorDesc SymbolicTensorDesc : AllSymbolicTensorDescs)
-	{
-		if (!SymbolicTensorDesc.GetShape().IsConcrete())
-		{
-			UE_LOG(LogNNE, Warning, TEXT("DML runtime does not support model with variable shapes yet."));
-			return -1;
-		}
-	}
-
-	return 0;
-}
 
 //
 // Called on RDG runtime startup
@@ -2506,9 +81,6 @@ bool FRuntimeDmlStartup()
 
 using namespace UE::NNERuntimeRDG::Private::Dml;
 
-FGuid UNNERuntimeRDGDmlImpl::GUID = FGuid((int32)'R', (int32)'D', (int32)'G', (int32)'D');
-int32 UNNERuntimeRDGDmlImpl::Version = 0x00000001;
-
 bool UNNERuntimeRDGDmlImpl::CanCreateModelData(FString FileType, TConstArrayView<uint8> FileData) const
 {
 	return FileType.Compare("onnx", ESearchCase::IgnoreCase) == 0;
@@ -2536,7 +108,10 @@ TArray<uint8> UNNERuntimeRDGDmlImpl::CreateModelData(FString FileType, TConstArr
 	TArray<uint8> Result;
 	FMemoryWriter Writer(Result);
 
-	Writer << GUID;
+	FGuid Guid = FModelInfo::Get()->GetGuid();
+	int32 Version = FModelInfo::Get()->GetVersion();
+
+	Writer << Guid;
 	Writer << Version;
 	Writer.Serialize(OutputModel.Data.GetData(), OutputModel.Data.Num());
 	return Result;
@@ -2545,16 +120,14 @@ TArray<uint8> UNNERuntimeRDGDmlImpl::CreateModelData(FString FileType, TConstArr
 bool UNNERuntimeRDGDmlImpl::CanCreateModelRDG(TObjectPtr<UNNEModelData> ModelData) const
 {
 #ifdef NNE_USE_DIRECTML
-	int32 GuidSize = sizeof(GUID);
-	int32 VersionSize = sizeof(Version);
 	TConstArrayView<uint8> Data = ModelData->GetModelData(GetRuntimeName());
 
-	if (Data.Num() <= GuidSize + VersionSize)
+	if (Data.Num() <= FModelInfo::Get()->GetGuidSize() + FModelInfo::Get()->GetVersionSize())
 	{
 		return false;
 	}
-	bool bResult = FGenericPlatformMemory::Memcmp(&(Data[0]), &(GUID), GuidSize) == 0;
-	bResult &= FGenericPlatformMemory::Memcmp(&(Data[GuidSize]), &(Version), VersionSize) == 0;
+
+	bool bResult = FModelInfo::Get()->ValidateGuidAndVersion(Data.GetData(), Data.GetData() + FModelInfo::Get()->GetGuidSize());
 	return bResult;
 #else
 	return false;
@@ -2564,21 +137,23 @@ bool UNNERuntimeRDGDmlImpl::CanCreateModelRDG(TObjectPtr<UNNEModelData> ModelDat
 TUniquePtr<UE::NNECore::IModelRDG> UNNERuntimeRDGDmlImpl::CreateModelRDG(TObjectPtr<UNNEModelData> ModelData)
 {
 #ifdef NNE_USE_DIRECTML
+	TUniquePtr<FModel>	Model;
+
 	if (!CanCreateModelRDG(ModelData))
 	{
-		return TUniquePtr<UE::NNECore::IModelRDG>();
+		return Model;
 	}
 
 	// Create the model and initialize it with the data not including the header
-	FModel* Model = new FModel();
+	Model = MakeUnique<FModel>();
 	TConstArrayView<uint8> Data = ModelData->GetModelData(GetRuntimeName());
 
 	if (!Model->Init(Data, Ctx))
 	{
-		delete Model;
-		return TUniquePtr<UE::NNECore::IModelRDG>();
+		Model.Reset();
 	}
-	return TUniquePtr<UE::NNECore::IModelRDG>(Model);
+
+	return Model;
 #else
 	return TUniquePtr<UE::NNECore::IModelRDG>();
 #endif
@@ -2598,12 +173,7 @@ UNNERuntimeRDGDmlImpl::~UNNERuntimeRDGDmlImpl()
 bool UNNERuntimeRDGDmlImpl::Init(bool bRegisterOnlyOperators)
 {
 #ifdef NNE_USE_DIRECTML
-	RegisterElementWiseUnaryOperators();
-	RegisterElementWiseBinaryOperators();
-	RegisterGemmOperator();
-	RegisterConvOperator();
-	RegisterConvTransposeOperator();
-
+	
 	if (bRegisterOnlyOperators)
 	{
 		UE_LOG(LogNNE, Display, TEXT("Registering only operators"));
@@ -2640,7 +210,7 @@ bool UNNERuntimeRDGDmlImpl::Init(bool bRegisterOnlyOperators)
 	}
 
 	check(Ctx == nullptr);
-	Ctx = new FDeviceContextDml();
+	Ctx = new FDmlDeviceContext();
 	Ctx->DeviceIndex = 0;
 	Ctx->D3D12Device = RHI->RHIGetDevice(Ctx->DeviceIndex);
 
diff --git a/Engine/Plugins/Experimental/NNE/Source/NNERuntimeRDG/Private/NNERuntimeRDGDml.h b/Engine/Plugins/Experimental/NNE/Source/NNERuntimeRDG/Private/NNERuntimeRDGDml.h
--- a/Engine/Plugins/Experimental/NNE/Source/NNERuntimeRDG/Private/NNERuntimeRDGDml.h
+++ b/Engine/Plugins/Experimental/NNE/Source/NNERuntimeRDG/Private/NNERuntimeRDGDml.h
@@ -14,7 +14,7 @@
 
 namespace UE::NNERuntimeRDG::Private::Dml
 {
-	class FDeviceContextDml;
+	class FDmlDeviceContext;
 	bool FRuntimeDmlStartup();
 }
 
@@ -25,8 +25,6 @@ class UNNERuntimeRDGDmlImpl : public UObject, public INNERuntime, public INNERun
 	GENERATED_BODY()
 
 public:
-	static FGuid GUID;
-	static int32 Version;
 
 	UNNERuntimeRDGDmlImpl():Ctx(nullptr) {};
 	virtual ~UNNERuntimeRDGDmlImpl();
@@ -42,6 +40,6 @@ public:
 	virtual bool CanCreateModelRDG(TObjectPtr<UNNEModelData> ModelData) const override;
 	virtual TUniquePtr<UE::NNECore::IModelRDG> CreateModelRDG(TObjectPtr<UNNEModelData> ModelData) override;
 private:
-	UE::NNERuntimeRDG::Private::Dml::FDeviceContextDml* Ctx;
+	UE::NNERuntimeRDG::Private::Dml::FDmlDeviceContext* Ctx;
 };
 
