Implement RDG DirectML Pad operator

#rnx

diff --git a/Engine/Plugins/Experimental/NNE/Source/NNEQA/Resources/NNEQATestDesc.json b/Engine/Plugins/Experimental/NNE/Source/NNEQA/Resources/NNEQATestDesc.json
--- a/Engine/Plugins/Experimental/NNE/Source/NNEQA/Resources/NNEQATestDesc.json
+++ b/Engine/Plugins/Experimental/NNE/Source/NNEQA/Resources/NNEQATestDesc.json
@@ -305,9 +305,6 @@
 		},
 		{
 			"category": "PadOp",
-			"runtimes": [
-				{ "name": "NNERuntimeRDGDml", "skip": true	}
-			],
 			"targets": [
 				{ "target": "Pad" }
 			]
diff --git a/Engine/Plugins/Experimental/NNE/Source/NNERuntimeRDG/Private/Dml/NNEDmlOperatorPad.cpp b/Engine/Plugins/Experimental/NNE/Source/NNERuntimeRDG/Private/Dml/NNEDmlOperatorPad.cpp
new file mode 100644
--- /dev/null
+++ b/Engine/Plugins/Experimental/NNE/Source/NNERuntimeRDG/Private/Dml/NNEDmlOperatorPad.cpp
@@ -0,0 +1,100 @@
+// Copyright Epic Games, Inc. All Rights Reserved.
+
+#ifdef NNE_USE_DIRECTML
+#include "NNEDmlOperator.h"
+
+namespace UE::NNERuntimeRDG::Private::Dml
+{
+
+class FOperatorDmlPad : public FOperatorDml
+{
+
+	static DML_PADDING_MODE ModeFromString(FStringView StringVal)
+	{
+		if (FCString::Stricmp(StringVal.GetData(), TEXT("CONSTANT")) == 0)
+		{
+			return DML_PADDING_MODE_CONSTANT;
+		}
+		else if (FCString::Stricmp(StringVal.GetData(), TEXT("REFLECT")) == 0)
+		{
+			return DML_PADDING_MODE_REFLECTION;
+		}
+		else if (FCString::Stricmp(StringVal.GetData(), TEXT("EDGE")) == 0)
+		{
+			return DML_PADDING_MODE_EDGE;
+		}
+		else
+		{
+			return DML_PADDING_MODE_CONSTANT;
+		}
+	}
+
+public:
+
+	static FOperatorDml* Create()
+	{
+		return new FOperatorDmlPad();
+	}
+
+	//
+	//
+	//
+	virtual bool Initialize(IDMLDevice* Device, TArrayView<const NNECore::Internal::FTensor> InputTensors, TArrayView<const NNECore::Internal::FTensor> OutputTensors, const NNECore::FAttributeMap& Attributes) override
+	{
+		check(InputTensors.Num() == 1);
+		check(OutputTensors.Num() == 1);
+
+		const NNECore::Internal::FTensor& InputTensor = InputTensors[0];
+		const NNECore::Internal::FTensor& OutputTensor = OutputTensors[0];
+
+		// Read attributes
+		float				Value;
+		DML_PADDING_MODE	Mode;
+		TArray<int32>		Pads;
+
+		Value = Attributes.GetValueOrDefault<float>(TEXT("value"), 0.0f);
+		Mode = ModeFromString(Attributes.GetValue<FString>(TEXT("mode")));
+		Pads = Attributes.GetValue<TArray<int32>>(TEXT("pads"));
+
+		if (InputTensor.GetShape().Rank() * 2 != Pads.Num())
+		{
+			UE_LOG(LogNNE, Warning, TEXT("pads attribute lenght (%d) should be twice the rank of input X (%d)."), Pads.Num(), InputTensor.GetShape().Rank());
+			return false;
+		}
+		
+		// Initialize tensor descriptors
+		DmlUtil::FTensorDesc	DmlInputTensor{};
+		DmlUtil::FTensorDesc	DmlOutputTensor{};
+
+		if (!InitDmlTensorDesc(DmlInputTensor, InputTensor))
+		{
+			UE_LOG(LogNNE, Warning, TEXT("Failed to initialize tensor(s) for DML inference"));
+			return false;
+		}
+
+		if (!InitDmlTensorDesc(DmlOutputTensor, OutputTensor))
+		{
+			UE_LOG(LogNNE, Warning, TEXT("Failed to initialize tensor(s) for DML inference"));
+			return false;
+		}
+
+		DML_PADDING_OPERATOR_DESC	PadOpDesc{};
+
+		PadOpDesc.InputTensor = &DmlInputTensor.Desc;
+		PadOpDesc.OutputTensor = &DmlOutputTensor.Desc;
+		PadOpDesc.PaddingMode = Mode;
+		PadOpDesc.PaddingValue = Value;
+		PadOpDesc.DimensionCount = Pads.Num() / 2;
+		PadOpDesc.StartPadding = (uint32*) Pads.GetData();
+		PadOpDesc.EndPadding = (uint32*) (Pads.GetData() + PadOpDesc.DimensionCount);
+
+		return CreateOperator(Device, DML_OPERATOR_DESC { DML_OPERATOR_PADDING, &PadOpDesc });
+	}
+};
+
+// Register Pad operator on Module startup
+NNE_DML_REGISTER_OP(Pad)
+
+} // namespace UE::NNERuntimeRDG::Private::Dml
+
+#endif // NNE_USE_DIRECTML
