NNE - Fix Warning in StaticAnalysis

diff --git a/Engine/Plugins/Experimental/NNE/Source/NNERuntimeRDG/Private/NNERuntimeRDGElementWiseUnaryHelper.cpp b/Engine/Plugins/Experimental/NNE/Source/NNERuntimeRDG/Private/NNERuntimeRDGElementWiseUnaryHelper.cpp
--- a/Engine/Plugins/Experimental/NNE/Source/NNERuntimeRDG/Private/NNERuntimeRDGElementWiseUnaryHelper.cpp
+++ b/Engine/Plugins/Experimental/NNE/Source/NNERuntimeRDG/Private/NNERuntimeRDGElementWiseUnaryHelper.cpp
@@ -13,7 +13,14 @@ namespace UE::NNERuntimeRDG::Private::ElementWiseUnaryCPUHelper
 		//https://mathworld.wolfram.com/InverseHyperbolicCosine.html
 		float FloatNan = FMath::Sqrt(-1.0f);
 		float yAboveOne = FMath::Loge(X + FMath::Sqrt(X + 1.0f) + FMath::Sqrt(X - 1.0f));
-		return (X == 1.0f) ? 0.0f : (X >= 1.0f) ? yAboveOne : FloatNan;
+		if (X == 1.0f)
+		{
+			return 0.0f;
+		}
+		else
+		{
+			return (X >= 1.0f) ? yAboveOne : FloatNan;
+		}
 	}
 	
 	template<> float Apply<NNECore::Internal::EElementWiseUnaryOperatorType::Asin>(float X, float Alpha, float Beta, float Gamma) { return FMath::Asin(X); }
@@ -101,7 +108,9 @@ namespace UE::NNERuntimeRDG::Private::ElementWiseUnaryCPUHelper
 
 	template<> float Apply<NNECore::Internal::EElementWiseUnaryOperatorType::Tanh>(float X, float Alpha, float Beta, float Gamma) {
 		//https://mathworld.wolfram.com/HyperbolicTangent.html
-		return Apply<NNECore::Internal::EElementWiseUnaryOperatorType::Sinh>(X, Alpha, Beta, Gamma) / Apply<NNECore::Internal::EElementWiseUnaryOperatorType::Cosh>(X, Alpha, Beta, Gamma);
+		float SinhValue = Apply<NNECore::Internal::EElementWiseUnaryOperatorType::Sinh>(X, Alpha, Beta, Gamma);
+		float CoshValue = Apply<NNECore::Internal::EElementWiseUnaryOperatorType::Cosh>(X, Alpha, Beta, Gamma);
+		return SinhValue / CoshValue;
 	}
 
 	template<> float Apply<NNECore::Internal::EElementWiseUnaryOperatorType::Erf>(float X, float Alpha, float Beta, float Gamma) {
