NNE - RDG HLSL runtime - Allow binary operator to forward tensor constness in graph

Problem:
* Constant tensors can be followed by binary operator before reaching a node needing constant inputs.
* We didn't had a way to test sheduling side shape inference helpers.

Solutions:
* Implement scheduling side tensor computation to propagate constness before the model is runned.
* Add unit tests for shape inference unary and binary op helpers.

Tests: all nne test green on editor win64 (but a convtranspose test as known issue) + style transfer demo

diff --git a/Engine/Plugins/Experimental/NNE/Source/NNECore/Internal/NNECoreTensor.h b/Engine/Plugins/Experimental/NNE/Source/NNECore/Internal/NNECoreTensor.h
--- a/Engine/Plugins/Experimental/NNE/Source/NNECore/Internal/NNECoreTensor.h
+++ b/Engine/Plugins/Experimental/NNE/Source/NNECore/Internal/NNECoreTensor.h
@@ -62,11 +62,11 @@ namespace UE::NNECore::Internal
 		template <typename T> void SetPreparedData(TConstArrayView<T> Data)
 		{
 			const uint8* DataPtr = reinterpret_cast<const uint8*>(Data.GetData());
-			TConstArrayView<uint8> DataAsByte = MakeArrayView(DataPtr, Data.Num() / sizeof(T));
+			TConstArrayView<uint8> DataAsByte = MakeArrayView(DataPtr, Data.Num() * sizeof(T));
 			
 			checkf(DataAsByte.Num() == DataSize, TEXT("Incorrect data size, it should match tensor shape and data type."));
 			PreparedData.Reset();
-			PreparedData.Append(Data);
+			PreparedData.Append(DataAsByte);
 		}
 
 		bool HasPreparedData() const
diff --git a/Engine/Plugins/Experimental/NNE/Source/NNEQA/NNEQA.Build.cs b/Engine/Plugins/Experimental/NNE/Source/NNEQA/NNEQA.Build.cs
--- a/Engine/Plugins/Experimental/NNE/Source/NNEQA/NNEQA.Build.cs
+++ b/Engine/Plugins/Experimental/NNE/Source/NNEQA/NNEQA.Build.cs
@@ -23,6 +23,7 @@ public class NNEQA : ModuleRules
 				"Json",
 				"CoreUObject",
 				"NNECore",
+				"NNERuntimeRDG",
 				"NNEUtils",
 				"RHI",
 				"RenderCore"
diff --git a/Engine/Plugins/Experimental/NNE/Source/NNEQA/Private/NNEQAUnitTestHelper.cpp b/Engine/Plugins/Experimental/NNE/Source/NNEQA/Private/NNEQAUnitTestHelper.cpp
new file mode 100644
--- /dev/null
+++ b/Engine/Plugins/Experimental/NNE/Source/NNEQA/Private/NNEQAUnitTestHelper.cpp
@@ -0,0 +1,116 @@
+// Copyright Epic Games, Inc. All Rights Reserved.
+
+#include "NNEQAUnitTestHelper.h"
+#include "NNECoreTypes.h"
+#include "NNERuntimeRDGElementWiseBinaryHelper.h"
+#include "NNERuntimeRDGElementWiseUnaryHelper.h"
+
+#if WITH_DEV_AUTOMATION_TESTS
+#include "Misc/AutomationTest.h"
+
+namespace UE::NNEQA::Private::NNERuntimeRDG
+{
+
+using namespace NNECore;
+using namespace NNECore::Internal;
+using namespace UE::NNERuntimeRDG::Internal;
+
+bool FUnitTestBase::RunTest(const FString& Parameters) { return false; }
+
+FShapeInferenceHelperUnitTestBase::FShapeInferenceHelperUnitTestBase(const FString& InClassName, const FString& InTestName, const FString& InSourceFile, int32 InSourceLine)
+	: FUnitTestBase(InClassName), TestName(InTestName), SourceFile(InSourceFile), SourceLine(InSourceLine) {}
+
+FString FShapeInferenceHelperUnitTestBase::GetTestSourceFileName() const { return SourceFile; }
+
+int32 FShapeInferenceHelperUnitTestBase::GetTestSourceFileLine() const { return SourceLine; }
+		
+FString FShapeInferenceHelperUnitTestBase::GetBeautifiedTestName() const { return TestName; };
+
+void FShapeInferenceHelperUnitTestBase::GetTests(TArray<FString>& OutBeautifiedNames, TArray <FString>& OutTestCommands) const
+{
+	OutBeautifiedNames.Add(TestName);
+	OutTestCommands.Add(FString());
+}
+
+bool FShapeInferenceHelperUnitTestBase::RunTest(const FString& Parameter) { return false; }
+
+FTensor FShapeInferenceHelperUnitTestBase::MakeTensor(FString Name, TConstArrayView<uint32> Shape)
+{
+	FTensorShape TensorShape = FTensorShape::Make(Shape);
+	FTensor Tensor = FTensor::Make(Name, TensorShape, ENNETensorDataType::Float);
+	return Tensor;
+}
+
+FTensor FShapeInferenceHelperUnitTestBase::MakeConstTensor(FString Name, TConstArrayView<uint32> Shape, TConstArrayView<float> Data)
+{
+	FTensor Tensor = MakeTensor(Name, Shape);
+	check(Data.Num() == Tensor.GetVolume());
+	Tensor.SetPreparedData(Data);
+	return Tensor;
+}
+
+bool FShapeInferenceHelperUnitTestBase::TestBinaryOutputIsOnlyComputedWhenItShould(EElementWiseBinaryOperatorType OpType)
+{
+	FTensor XC1 = MakeConstTensor(TEXT("XC1"), { 1 }, { 1.0f });
+	FTensor XC20 = MakeConstTensor(TEXT("XC20"), { 20 }, { 3.0f, 4.0f, 3.0f, 4.0f, 3.0f, 3.0f, 4.0f, 3.0f, 4.0f, 3.0f, 3.0f, 4.0f, 3.0f, 4.0f, 3.0f, 3.0f, 4.0f, 3.0f, 4.0f, 3.0f });
+	FTensor X1 = MakeTensor(TEXT("X"), { 1 });
+
+	//Tests if output tensors are only computed if both inputs are constant and below a certain size
+	
+	FTensor Y = MakeTensor(TEXT("Y"), { 1 });
+	ElementWiseBinaryCPUHelper::Apply(OpType, XC1, XC1, Y);
+	UTEST_TRUE(TEXT("Y const if both inputs are const"), Y.HasPreparedData());
+
+	Y = MakeTensor(TEXT("Y"), { 1 });
+	ElementWiseBinaryCPUHelper::Apply(OpType, X1, XC1, Y);
+	UTEST_FALSE(TEXT("Y not const when LHS not const"), Y.HasPreparedData());
+
+	Y = MakeTensor(TEXT("Y"), { 1 });
+	ElementWiseBinaryCPUHelper::Apply(OpType, XC1, X1, Y);
+	UTEST_FALSE(TEXT("Y not const when RHS not const"), Y.HasPreparedData());
+
+	Y = MakeTensor(TEXT("Y"), { 1 });
+	ElementWiseBinaryCPUHelper::Apply(OpType, X1, X1, Y);
+	UTEST_FALSE(TEXT("Y not const when LHS and RHS not const"), Y.HasPreparedData());
+
+	Y = MakeTensor(TEXT("Y"), { 20 });
+	ElementWiseBinaryCPUHelper::Apply(OpType, XC20, XC1, Y);
+	UTEST_FALSE(TEXT("Y not const when LHS too large"), Y.HasPreparedData());
+
+	Y = MakeTensor(TEXT("Y"), { 20 });
+	ElementWiseBinaryCPUHelper::Apply(OpType, XC1, XC20, Y);
+	UTEST_FALSE(TEXT("Y not const when RHS too large"), Y.HasPreparedData());
+
+	Y = MakeTensor(TEXT("Y"), { 20 });
+	ElementWiseBinaryCPUHelper::Apply(OpType, XC20, XC20, Y);
+	UTEST_FALSE(TEXT("Y not const when LHS and RHS too large"), Y.HasPreparedData());
+
+	return true;
+}
+
+bool FShapeInferenceHelperUnitTestBase::TestUnaryOutputIsOnlyComputedWhenItShould(EElementWiseUnaryOperatorType OpType)
+{
+	FTensor XC1 = MakeConstTensor(TEXT("XC1"), { 1 }, { 1.0f });
+	FTensor XC20 = MakeConstTensor(TEXT("XC20"), { 20 }, { 3.0f, 4.0f, 3.0f, 4.0f, 3.0f, 3.0f, 4.0f, 3.0f, 4.0f, 3.0f, 3.0f, 4.0f, 3.0f, 4.0f, 3.0f, 3.0f, 4.0f, 3.0f, 4.0f, 3.0f });
+	FTensor X1 = MakeTensor(TEXT("X"), { 1 });
+
+	//Tests if output tensors are only computed if inputs are constant and below a certain size
+
+	FTensor Y = MakeTensor(TEXT("Y"), { 1 });
+	ElementWiseUnaryCPUHelper::Apply(OpType, XC1, 1.0f, 1.0f, 1.0f, Y);
+	UTEST_TRUE(TEXT("Y const if input is const"), Y.HasPreparedData());
+
+	Y = MakeTensor(TEXT("Y"), { 1 });
+	ElementWiseUnaryCPUHelper::Apply(OpType, X1, 1.0f, 1.0f, 1.0f, Y);
+	UTEST_FALSE(TEXT("Y not const if input is not const"), Y.HasPreparedData());
+
+	Y = MakeTensor(TEXT("Y"), { 20 });
+	ElementWiseUnaryCPUHelper::Apply(OpType, XC20, 1.0f, 1.0f, 1.0f, Y);
+	UTEST_FALSE(TEXT("Y not const if input is too large "), Y.HasPreparedData());
+
+	return true;
+}
+
+} // namespace UE::NNEQA::Private::NNERuntimeRDG
+
+#endif //WITH_DEV_AUTOMATION_TESTS
diff --git a/Engine/Plugins/Experimental/NNE/Source/NNEQA/Private/NNEQAUnitTestHelper.h b/Engine/Plugins/Experimental/NNE/Source/NNEQA/Private/NNEQAUnitTestHelper.h
new file mode 100644
--- /dev/null
+++ b/Engine/Plugins/Experimental/NNE/Source/NNEQA/Private/NNEQAUnitTestHelper.h
@@ -0,0 +1,64 @@
+// Copyright Epic Games, Inc. All Rights Reserved.
+
+#pragma once
+
+#include "CoreMinimal.h"
+#include "NNECoreTensor.h"
+#include "NNECoreOperator.h"
+
+#if WITH_DEV_AUTOMATION_TESTS
+#include "Misc/AutomationTest.h"
+
+namespace UE::NNEQA::Private::NNERuntimeRDG
+{
+	using namespace NNECore::Internal;
+
+	IMPLEMENT_SIMPLE_AUTOMATION_TEST_PRIVATE(FUnitTestBase, FAutomationTestBase, "nne.UnitTest.Base", EAutomationTestFlags::ApplicationContextMask | EAutomationTestFlags::FeatureMask | EAutomationTestFlags::EngineFilter, __FILE__, __LINE__)
+
+	class FShapeInferenceHelperUnitTestBase : public FUnitTestBase
+	{
+		//Automation tests
+		FString TestName;
+		FString SourceFile;
+		int32 SourceLine;
+		
+	public:
+		FShapeInferenceHelperUnitTestBase(const FString& InClassName, const FString& InTestName, const FString& InSourceFile, int32 InSourceLine);
+		virtual ~FShapeInferenceHelperUnitTestBase() {}
+		virtual FString GetTestSourceFileName() const override;
+		virtual int32 GetTestSourceFileLine() const override;
+		
+	protected:
+		virtual FString GetBeautifiedTestName() const override;
+		virtual void GetTests(TArray<FString>& OutBeautifiedNames, TArray <FString>& OutTestCommands) const override;
+		virtual bool RunTest(const FString& Parameter) override;
+
+		//Shape inference unit test helpers
+		static FTensor MakeTensor(FString Name, TConstArrayView<uint32> Shape);
+		static FTensor MakeConstTensor(FString Name, TConstArrayView<uint32> Shape, TConstArrayView<float> Data);
+		bool TestUnaryOutputIsOnlyComputedWhenItShould(EElementWiseUnaryOperatorType OpType);
+		bool TestBinaryOutputIsOnlyComputedWhenItShould(EElementWiseBinaryOperatorType OpType);
+		
+	};
+
+	#define IMPLEMENT_NNE_SHAPEINFERENCEHELPER_UNIT_AUTOMATION_TEST( TClass, PrettyName ) \
+	class TClass : public FShapeInferenceHelperUnitTestBase \
+	{ \
+	public: \
+		TClass(const FString& InClassName, const FString& InTestName, const FString& InSourceFile, int32 InSourceLine) \
+			: FShapeInferenceHelperUnitTestBase(InClassName, InTestName, InSourceFile, InSourceLine) {} \
+		virtual ~TClass() {} \
+    \
+	protected: \
+		virtual bool RunTest(const FString& Parameter) override; \
+	}; \
+	namespace\
+	{\
+		TClass TClass##AutomationTestInstance( TEXT(#TClass), PrettyName, __FILE__, __LINE__ );\
+	}
+
+
+
+} // namespace UE::NNEQA::Private::NNERuntimeRDG
+
+#endif //WITH_DEV_AUTOMATION_TESTS
diff --git a/Engine/Plugins/Experimental/NNE/Source/NNEQA/Private/NNEQAUnitTestOpBinary.cpp b/Engine/Plugins/Experimental/NNE/Source/NNEQA/Private/NNEQAUnitTestOpBinary.cpp
new file mode 100644
--- /dev/null
+++ b/Engine/Plugins/Experimental/NNE/Source/NNEQA/Private/NNEQAUnitTestOpBinary.cpp
@@ -0,0 +1,285 @@
+// Copyright Epic Games, Inc. All Rights Reserved.
+
+#include "CoreMinimal.h"
+#include "NNECoreTensor.h"
+#include "NNECoreTypes.h"
+#include "NNEQAUnitTestHelper.h"
+#include "NNERuntimeRDGElementWiseBinaryHelper.h"
+
+#if WITH_DEV_AUTOMATION_TESTS
+#include "Misc/AutomationTest.h"
+#endif
+
+namespace UE::NNEQA::Private::NNERuntimeRDG::BinaryOp
+{
+
+#if WITH_DEV_AUTOMATION_TESTS
+	
+	using namespace NNECore;
+	using namespace NNECore::Internal;
+	using namespace UE::NNERuntimeRDG::Internal;
+
+	IMPLEMENT_NNE_SHAPEINFERENCEHELPER_UNIT_AUTOMATION_TEST(FElementWiseBinaryCPUHelperAdd, "System.Engine.MachineLearning.NNE.UnitTest.BinaryHelper.Add");
+	bool FElementWiseBinaryCPUHelperAdd::RunTest(const FString& Parameter)
+	{
+		FTensor XC1 = MakeConstTensor(TEXT("XC1"), { 1 }, { 1.0f });
+		FTensor XC1x2 = MakeConstTensor(TEXT("XC1x2"), { 1,2 }, { 1.0f, 2.0f });
+		FTensor XC2x1 = MakeConstTensor(TEXT("XC2x1"), { 2,1 }, { 3.0f, 4.0f });
+
+		if (!TestBinaryOutputIsOnlyComputedWhenItShould(EElementWiseBinaryOperatorType::Add))
+		{
+			return false;
+		}
+
+		//Test output tensor math is correct when both inputs are constant including broadcasting
+		FTensor Y = MakeTensor(TEXT("Y"), { 1 });
+		ElementWiseBinaryCPUHelper::Apply(EElementWiseBinaryOperatorType::Add, XC1, XC1, Y);
+		UTEST_EQUAL(TEXT("XC1+XC1"), Y.GetPreparedData<float>()[0], 1.0f + 1.0f);
+
+		Y = MakeTensor(TEXT("Y"), { 1,2 });
+		ElementWiseBinaryCPUHelper::Apply(EElementWiseBinaryOperatorType::Add, XC1x2, XC1, Y);
+		UTEST_EQUAL(TEXT("XC1x2+XC1[0]"), Y.GetPreparedData<float>()[0], 1.0f + 1.0f);
+		UTEST_EQUAL(TEXT("XC1x2+XC1[1]"), Y.GetPreparedData<float>()[1], 2.0f + 1.0f);
+
+		Y = MakeTensor(TEXT("Y"), { 1,2 });
+		ElementWiseBinaryCPUHelper::Apply(EElementWiseBinaryOperatorType::Add, XC1x2, XC1x2, Y);
+		UTEST_EQUAL(TEXT("XC1x2+XC1x2[0]"), Y.GetPreparedData<float>()[0], 1.0f + 1.0f);
+		UTEST_EQUAL(TEXT("XC1x2+XC1x2[1]"), Y.GetPreparedData<float>()[1], 2.0f + 2.0f);
+
+		Y = MakeTensor(TEXT("Y"), { 2,2 });
+		ElementWiseBinaryCPUHelper::Apply(EElementWiseBinaryOperatorType::Add, XC1x2, XC2x1, Y);
+		UTEST_EQUAL(TEXT("XC1x2+XC2x1[0]"), Y.GetPreparedData<float>()[0], 1.0f + 3.0f);
+		UTEST_EQUAL(TEXT("XC1x2+XC2x1[1]"), Y.GetPreparedData<float>()[1], 2.0f + 3.0f);
+		UTEST_EQUAL(TEXT("XC1x2+XC2x1[2]"), Y.GetPreparedData<float>()[2], 1.0f + 4.0f);
+		UTEST_EQUAL(TEXT("XC1x2+XC2x1[3]"), Y.GetPreparedData<float>()[3], 2.0f + 4.0f);
+
+		return true;
+	}
+
+	IMPLEMENT_NNE_SHAPEINFERENCEHELPER_UNIT_AUTOMATION_TEST(FElementWiseBinaryCPUHelperDiv, "System.Engine.MachineLearning.NNE.UnitTest.BinaryHelper.Div");
+	bool FElementWiseBinaryCPUHelperDiv::RunTest(const FString& Parameter)
+	{
+		FTensor XC1 = MakeConstTensor(TEXT("XC1"), { 1 }, { 1.0f });
+		FTensor XC1x2 = MakeConstTensor(TEXT("XC1x2"), { 1,2 }, { 1.0f, 2.0f });
+		FTensor XC2x1 = MakeConstTensor(TEXT("XC2x1"), { 2,1 }, { 3.0f, 4.0f });
+
+		if (!TestBinaryOutputIsOnlyComputedWhenItShould(EElementWiseBinaryOperatorType::Div))
+		{
+			return false;
+		}
+
+		//Test output tensor math is correct when both inputs are constant including broadcasting
+		FTensor Y = MakeTensor(TEXT("Y"), { 1 });
+		ElementWiseBinaryCPUHelper::Apply(EElementWiseBinaryOperatorType::Div, XC1, XC1, Y);
+		UTEST_EQUAL(TEXT("XC1/XC1"), Y.GetPreparedData<float>()[0], 1.0f / 1.0f);
+
+		Y = MakeTensor(TEXT("Y"), { 1,2 });
+		ElementWiseBinaryCPUHelper::Apply(EElementWiseBinaryOperatorType::Div, XC1x2, XC1, Y);
+		UTEST_EQUAL(TEXT("XC1x2/XC1[0]"), Y.GetPreparedData<float>()[0], 1.0f / 1.0f);
+		UTEST_EQUAL(TEXT("XC1x2/XC1[1]"), Y.GetPreparedData<float>()[1], 2.0f / 1.0f);
+
+		Y = MakeTensor(TEXT("Y"), { 1,2 });
+		ElementWiseBinaryCPUHelper::Apply(EElementWiseBinaryOperatorType::Div, XC1x2, XC1x2, Y);
+		UTEST_EQUAL(TEXT("XC1x2/XC1x2[0]"), Y.GetPreparedData<float>()[0], 1.0f / 1.0f);
+		UTEST_EQUAL(TEXT("XC1x2/XC1x2[1]"), Y.GetPreparedData<float>()[1], 2.0f / 2.0f);
+
+		Y = MakeTensor(TEXT("Y"), { 2,2 });
+		ElementWiseBinaryCPUHelper::Apply(EElementWiseBinaryOperatorType::Div, XC1x2, XC2x1, Y);
+		UTEST_EQUAL(TEXT("XC1x2/XC2x1[0]"), Y.GetPreparedData<float>()[0], 1.0f / 3.0f);
+		UTEST_EQUAL(TEXT("XC1x2/XC2x1[1]"), Y.GetPreparedData<float>()[1], 2.0f / 3.0f);
+		UTEST_EQUAL(TEXT("XC1x2/XC2x1[2]"), Y.GetPreparedData<float>()[2], 1.0f / 4.0f);
+		UTEST_EQUAL(TEXT("XC1x2/XC2x1[3]"), Y.GetPreparedData<float>()[3], 2.0f / 4.0f);
+
+		return true;
+	};
+
+	IMPLEMENT_NNE_SHAPEINFERENCEHELPER_UNIT_AUTOMATION_TEST(FElementWiseBinaryCPUHelperMul, "System.Engine.MachineLearning.NNE.UnitTest.BinaryHelper.Mul");
+	bool FElementWiseBinaryCPUHelperMul::RunTest(const FString& Parameter)
+	{
+		FTensor XC1 = MakeConstTensor(TEXT("XC1"), { 1 }, { 1.0f });
+		FTensor XC1x2 = MakeConstTensor(TEXT("XC1x2"), { 1,2 }, { 1.0f, 2.0f });
+		FTensor XC2x1 = MakeConstTensor(TEXT("XC2x1"), { 2,1 }, { 3.0f, 4.0f });
+
+		if (!TestBinaryOutputIsOnlyComputedWhenItShould(EElementWiseBinaryOperatorType::Mul))
+		{
+			return false;
+		}
+
+		//Test output tensor math is correct when both inputs are constant including broadcasting
+		FTensor Y = MakeTensor(TEXT("Y"), { 1 });
+		ElementWiseBinaryCPUHelper::Apply(EElementWiseBinaryOperatorType::Mul, XC1, XC1, Y);
+		UTEST_EQUAL(TEXT("XC1*XC1"), Y.GetPreparedData<float>()[0], 1.0f * 1.0f);
+
+		Y = MakeTensor(TEXT("Y"), { 1,2 });
+		ElementWiseBinaryCPUHelper::Apply(EElementWiseBinaryOperatorType::Mul, XC1x2, XC1, Y);
+		UTEST_EQUAL(TEXT("XC1x2*XC1[0]"), Y.GetPreparedData<float>()[0], 1.0f * 1.0f);
+		UTEST_EQUAL(TEXT("XC1x2*XC1[1]"), Y.GetPreparedData<float>()[1], 2.0f * 1.0f);
+
+		Y = MakeTensor(TEXT("Y"), { 1,2 });
+		ElementWiseBinaryCPUHelper::Apply(EElementWiseBinaryOperatorType::Mul, XC1x2, XC1x2, Y);
+		UTEST_EQUAL(TEXT("XC1x2*XC1x2[0]"), Y.GetPreparedData<float>()[0], 1.0f * 1.0f);
+		UTEST_EQUAL(TEXT("XC1x2*XC1x2[1]"), Y.GetPreparedData<float>()[1], 2.0f * 2.0f);
+
+		Y = MakeTensor(TEXT("Y"), { 2,2 });
+		ElementWiseBinaryCPUHelper::Apply(EElementWiseBinaryOperatorType::Mul, XC1x2, XC2x1, Y);
+		UTEST_EQUAL(TEXT("XC1x2*XC2x1[0]"), Y.GetPreparedData<float>()[0], 1.0f * 3.0f);
+		UTEST_EQUAL(TEXT("XC1x2*XC2x1[1]"), Y.GetPreparedData<float>()[1], 2.0f * 3.0f);
+		UTEST_EQUAL(TEXT("XC1x2*XC2x1[2]"), Y.GetPreparedData<float>()[2], 1.0f * 4.0f);
+		UTEST_EQUAL(TEXT("XC1x2*XC2x1[3]"), Y.GetPreparedData<float>()[3], 2.0f * 4.0f);
+
+		return true;
+	};
+
+	IMPLEMENT_NNE_SHAPEINFERENCEHELPER_UNIT_AUTOMATION_TEST(FElementWiseBinaryCPUHelperSub, "System.Engine.MachineLearning.NNE.UnitTest.BinaryHelper.Sub");
+	bool FElementWiseBinaryCPUHelperSub::RunTest(const FString& Parameter)
+	{
+		FTensor XC1 = MakeConstTensor(TEXT("XC1"), { 1 }, { 1.0f });
+		FTensor XC1x2 = MakeConstTensor(TEXT("XC1x2"), { 1,2 }, { 1.0f, 2.0f });
+		FTensor XC2x1 = MakeConstTensor(TEXT("XC2x1"), { 2,1 }, { 3.0f, 4.0f });
+
+		if (!TestBinaryOutputIsOnlyComputedWhenItShould(EElementWiseBinaryOperatorType::Sub))
+		{
+			return false;
+		}
+
+		//Test output tensor math is correct when both inputs are constant including broadcasting
+		FTensor Y = MakeTensor(TEXT("Y"), { 1 });
+		ElementWiseBinaryCPUHelper::Apply(EElementWiseBinaryOperatorType::Sub, XC1, XC1, Y);
+		UTEST_EQUAL(TEXT("XC1-XC1"), Y.GetPreparedData<float>()[0], 1.0f - 1.0f);
+
+		Y = MakeTensor(TEXT("Y"), { 1,2 });
+		ElementWiseBinaryCPUHelper::Apply(EElementWiseBinaryOperatorType::Sub, XC1x2, XC1, Y);
+		UTEST_EQUAL(TEXT("XC1x2-XC1[0]"), Y.GetPreparedData<float>()[0], 1.0f - 1.0f);
+		UTEST_EQUAL(TEXT("XC1x2-XC1[1]"), Y.GetPreparedData<float>()[1], 2.0f - 1.0f);
+
+		Y = MakeTensor(TEXT("Y"), { 1,2 });
+		ElementWiseBinaryCPUHelper::Apply(EElementWiseBinaryOperatorType::Sub, XC1x2, XC1x2, Y);
+		UTEST_EQUAL(TEXT("XC1x2-XC1x2[0]"), Y.GetPreparedData<float>()[0], 1.0f - 1.0f);
+		UTEST_EQUAL(TEXT("XC1x2-XC1x2[1]"), Y.GetPreparedData<float>()[1], 2.0f - 2.0f);
+
+		Y = MakeTensor(TEXT("Y"), { 2,2 });
+		ElementWiseBinaryCPUHelper::Apply(EElementWiseBinaryOperatorType::Sub, XC1x2, XC2x1, Y);
+		UTEST_EQUAL(TEXT("XC1x2-XC2x1[0]"), Y.GetPreparedData<float>()[0], 1.0f - 3.0f);
+		UTEST_EQUAL(TEXT("XC1x2-XC2x1[1]"), Y.GetPreparedData<float>()[1], 2.0f - 3.0f);
+		UTEST_EQUAL(TEXT("XC1x2-XC2x1[2]"), Y.GetPreparedData<float>()[2], 1.0f - 4.0f);
+		UTEST_EQUAL(TEXT("XC1x2-XC2x1[3]"), Y.GetPreparedData<float>()[3], 2.0f - 4.0f);
+
+		return true;
+	};
+
+	IMPLEMENT_NNE_SHAPEINFERENCEHELPER_UNIT_AUTOMATION_TEST(FElementWiseBinaryCPUHelperMod, "System.Engine.MachineLearning.NNE.UnitTest.BinaryHelper.Mod");
+	bool FElementWiseBinaryCPUHelperMod::RunTest(const FString& Parameter)
+	{
+		FTensor XC1 = MakeConstTensor(TEXT("XC1"), { 1 }, { 1.0f });
+		FTensor XC1x2 = MakeConstTensor(TEXT("XC1x2"), { 1,2 }, { 1.0f, 2.0f });
+		FTensor XC2x1 = MakeConstTensor(TEXT("XC2x1"), { 2,1 }, { 3.0f, 4.0f });
+
+		if (!TestBinaryOutputIsOnlyComputedWhenItShould(EElementWiseBinaryOperatorType::Mod))
+		{
+			return false;
+		}
+
+		//Test output tensor math is correct when both inputs are constant including broadcasting
+		FTensor Y = MakeTensor(TEXT("Y"), { 1 });
+		ElementWiseBinaryCPUHelper::Apply(EElementWiseBinaryOperatorType::Mod, XC1, XC1, Y);
+		UTEST_EQUAL(TEXT("XC1%XC1"), Y.GetPreparedData<float>()[0], FMath::Fmod(1.0f, 1.0f));
+
+		Y = MakeTensor(TEXT("Y"), { 1,2 });
+		ElementWiseBinaryCPUHelper::Apply(EElementWiseBinaryOperatorType::Mod, XC1x2, XC1, Y);
+		UTEST_EQUAL(TEXT("XC1x2%XC1[0]"), Y.GetPreparedData<float>()[0], FMath::Fmod(1.0f, 1.0f));
+		UTEST_EQUAL(TEXT("XC1x2%XC1[1]"), Y.GetPreparedData<float>()[1], FMath::Fmod(2.0f, 1.0f));
+
+		Y = MakeTensor(TEXT("Y"), { 1,2 });
+		ElementWiseBinaryCPUHelper::Apply(EElementWiseBinaryOperatorType::Mod, XC1x2, XC1x2, Y);
+		UTEST_EQUAL(TEXT("XC1x2%XC1x2[0]"), Y.GetPreparedData<float>()[0], FMath::Fmod(1.0f, 1.0f));
+		UTEST_EQUAL(TEXT("XC1x2%XC1x2[1]"), Y.GetPreparedData<float>()[1], FMath::Fmod(2.0f, 2.0f));
+
+		Y = MakeTensor(TEXT("Y"), { 2,2 });
+		ElementWiseBinaryCPUHelper::Apply(EElementWiseBinaryOperatorType::Mod, XC1x2, XC2x1, Y);
+		UTEST_EQUAL(TEXT("XC1x2%XC2x1[0]"), Y.GetPreparedData<float>()[0], FMath::Fmod(1.0f, 3.0f));
+		UTEST_EQUAL(TEXT("XC1x2%XC2x1[1]"), Y.GetPreparedData<float>()[1], FMath::Fmod(2.0f, 3.0f));
+		UTEST_EQUAL(TEXT("XC1x2%XC2x1[2]"), Y.GetPreparedData<float>()[2], FMath::Fmod(1.0f, 4.0f));
+		UTEST_EQUAL(TEXT("XC1x2%XC2x1[3]"), Y.GetPreparedData<float>()[3], FMath::Fmod(2.0f, 4.0f));
+
+		return true;
+	};
+
+	IMPLEMENT_NNE_SHAPEINFERENCEHELPER_UNIT_AUTOMATION_TEST(FElementWiseBinaryCPUHelperPow, "System.Engine.MachineLearning.NNE.UnitTest.BinaryHelper.Pow");
+	bool FElementWiseBinaryCPUHelperPow::RunTest(const FString& Parameter)
+	{
+		FTensor XC1 = MakeConstTensor(TEXT("XC1"), { 1 }, { 1.0f });
+		FTensor XC1x2 = MakeConstTensor(TEXT("XC1x2"), { 1,2 }, { 1.0f, 2.0f });
+		FTensor XC2x1 = MakeConstTensor(TEXT("XC2x1"), { 2,1 }, { 3.0f, 4.0f });
+
+		if (!TestBinaryOutputIsOnlyComputedWhenItShould(EElementWiseBinaryOperatorType::Pow))
+		{
+			return false;
+		}
+
+		//Test output tensor math is correct when both inputs are constant including broadcasting
+		FTensor Y = MakeTensor(TEXT("Y"), { 1 });
+		ElementWiseBinaryCPUHelper::Apply(EElementWiseBinaryOperatorType::Pow, XC1, XC1, Y);
+		UTEST_EQUAL(TEXT("XC1^XC1"), Y.GetPreparedData<float>()[0], FMath::Pow(1.0f, 1.0f));
+
+		Y = MakeTensor(TEXT("Y"), { 1,2 });
+		ElementWiseBinaryCPUHelper::Apply(EElementWiseBinaryOperatorType::Pow, XC1x2, XC1, Y);
+		UTEST_EQUAL(TEXT("XC1x2^XC1[0]"), Y.GetPreparedData<float>()[0], FMath::Pow(1.0f, 1.0f));
+		UTEST_EQUAL(TEXT("XC1x2^XC1[1]"), Y.GetPreparedData<float>()[1], FMath::Pow(2.0f, 1.0f));
+
+		Y = MakeTensor(TEXT("Y"), { 1,2 });
+		ElementWiseBinaryCPUHelper::Apply(EElementWiseBinaryOperatorType::Pow, XC1x2, XC1x2, Y);
+		UTEST_EQUAL(TEXT("XC1x2^XC1x2[0]"), Y.GetPreparedData<float>()[0], FMath::Pow(1.0f, 1.0f));
+		UTEST_EQUAL(TEXT("XC1x2^XC1x2[1]"), Y.GetPreparedData<float>()[1], FMath::Pow(2.0f, 2.0f));
+
+		Y = MakeTensor(TEXT("Y"), { 2,2 });
+		ElementWiseBinaryCPUHelper::Apply(EElementWiseBinaryOperatorType::Pow, XC1x2, XC2x1, Y);
+		UTEST_EQUAL(TEXT("XC1x2^XC2x1[0]"), Y.GetPreparedData<float>()[0], FMath::Pow(1.0f, 3.0f));
+		UTEST_EQUAL(TEXT("XC1x2^XC2x1[1]"), Y.GetPreparedData<float>()[1], FMath::Pow(2.0f, 3.0f));
+		UTEST_EQUAL(TEXT("XC1x2^XC2x1[2]"), Y.GetPreparedData<float>()[2], FMath::Pow(1.0f, 4.0f));
+		UTEST_EQUAL(TEXT("XC1x2^XC2x1[3]"), Y.GetPreparedData<float>()[3], FMath::Pow(2.0f, 4.0f));
+
+		return true;
+	};
+
+	IMPLEMENT_NNE_SHAPEINFERENCEHELPER_UNIT_AUTOMATION_TEST(FElementWiseBinaryCPUHelperPrelu, "System.Engine.MachineLearning.NNE.UnitTest.BinaryHelper.Prelu");
+	bool FElementWiseBinaryCPUHelperPrelu::RunTest(const FString& Parameter)
+	{
+		FTensor XC1 = MakeConstTensor(TEXT("XC1"), { 1 }, { 1.0f });
+		FTensor XC1x2 = MakeConstTensor(TEXT("XC1x2"), { 1,2 }, { 1.0f, 2.0f });
+		FTensor XC2x1 = MakeConstTensor(TEXT("XC2x1"), { 2,1 }, { 3.0f, 4.0f });
+
+		if (!TestBinaryOutputIsOnlyComputedWhenItShould(EElementWiseBinaryOperatorType::Prelu))
+		{
+			return false;
+		}
+
+		auto Prelu = [](float X, float Y) { return (X < 0.0f) ? (Y * X) : X; };
+
+		//Test output tensor math is correct when both inputs are constant including broadcasting
+		FTensor Y = MakeTensor(TEXT("Y"), { 1 });
+		ElementWiseBinaryCPUHelper::Apply(EElementWiseBinaryOperatorType::Prelu, XC1, XC1, Y);
+		UTEST_EQUAL(TEXT("XC1_XC1"), Y.GetPreparedData<float>()[0], Prelu(1.0f, 1.0f));
+
+		Y = MakeTensor(TEXT("Y"), { 1,2 });
+		ElementWiseBinaryCPUHelper::Apply(EElementWiseBinaryOperatorType::Prelu, XC1x2, XC1, Y);
+		UTEST_EQUAL(TEXT("XC1x2_XC1[0]"), Y.GetPreparedData<float>()[0], Prelu(1.0f, 1.0f));
+		UTEST_EQUAL(TEXT("XC1x2_XC1[1]"), Y.GetPreparedData<float>()[1], Prelu(2.0f, 1.0f));
+
+		Y = MakeTensor(TEXT("Y"), { 1,2 });
+		ElementWiseBinaryCPUHelper::Apply(EElementWiseBinaryOperatorType::Prelu, XC1x2, XC1x2, Y);
+		UTEST_EQUAL(TEXT("XC1x2_XC1x2[0]"), Y.GetPreparedData<float>()[0], Prelu(1.0f, 1.0f));
+		UTEST_EQUAL(TEXT("XC1x2_XC1x2[1]"), Y.GetPreparedData<float>()[1], Prelu(2.0f, 2.0f));
+
+		Y = MakeTensor(TEXT("Y"), { 2,2 });
+		ElementWiseBinaryCPUHelper::Apply(EElementWiseBinaryOperatorType::Prelu, XC1x2, XC2x1, Y);
+		UTEST_EQUAL(TEXT("XC1x2_XC2x1[0]"), Y.GetPreparedData<float>()[0], Prelu(1.0f, 3.0f));
+		UTEST_EQUAL(TEXT("XC1x2_XC2x1[1]"), Y.GetPreparedData<float>()[1], Prelu(2.0f, 3.0f));
+		UTEST_EQUAL(TEXT("XC1x2_XC2x1[2]"), Y.GetPreparedData<float>()[2], Prelu(1.0f, 4.0f));
+		UTEST_EQUAL(TEXT("XC1x2_XC2x1[3]"), Y.GetPreparedData<float>()[3], Prelu(2.0f, 4.0f));
+
+		return true;
+	};
+
+#endif //WITH_DEV_AUTOMATION_TESTS
+
+} // namespace UE::NNEQA::Private::NNERuntimeRDG::BinaryOp
diff --git a/Engine/Plugins/Experimental/NNE/Source/NNEQA/Private/NNEQAUnitTestOpUnary.cpp b/Engine/Plugins/Experimental/NNE/Source/NNEQA/Private/NNEQAUnitTestOpUnary.cpp
new file mode 100644
--- /dev/null
+++ b/Engine/Plugins/Experimental/NNE/Source/NNEQA/Private/NNEQAUnitTestOpUnary.cpp
@@ -0,0 +1,391 @@
+// Copyright Epic Games, Inc. All Rights Reserved.
+
+#include "CoreMinimal.h"
+#include "Misc/EnumRange.h"
+#include "NNECoreTensor.h"
+#include "NNECoreTypes.h"
+#include "NNEQAUnitTestHelper.h"
+#include "NNERuntimeRDGElementWiseUnaryHelper.h"
+
+#if WITH_DEV_AUTOMATION_TESTS
+#include "Misc/AutomationTest.h"
+#endif
+
+ENUM_RANGE_BY_COUNT(UE::NNECore::Internal::EElementWiseUnaryOperatorType, UE::NNECore::Internal::EElementWiseUnaryOperatorType::MAX)
+
+namespace UE::NNEQA::Private::NNERuntimeRDG::UnaryOp
+{
+
+#if WITH_DEV_AUTOMATION_TESTS
+
+	using namespace NNECore;
+	using namespace NNECore::Internal;
+	using namespace UE::NNERuntimeRDG::Internal;
+
+	IMPLEMENT_NNE_SHAPEINFERENCEHELPER_UNIT_AUTOMATION_TEST(FElementWiseUnaryCPUHelperConstOuput, "System.Engine.MachineLearning.NNE.UnitTest.UnaryHelper.ConstOutput");
+	bool FElementWiseUnaryCPUHelperConstOuput::RunTest(const FString& Parameter)
+	{
+		for (EElementWiseUnaryOperatorType Op : TEnumRange<EElementWiseUnaryOperatorType>())
+		{
+			if (!TestUnaryOutputIsOnlyComputedWhenItShould(Op))
+			{
+				return false;
+			}
+		}
+		return true;
+	}
+
+	IMPLEMENT_NNE_SHAPEINFERENCEHELPER_UNIT_AUTOMATION_TEST(FElementWiseUnaryCPUHelperAbs, "System.Engine.MachineLearning.NNE.UnitTest.UnaryHelper.Abs");
+	bool FElementWiseUnaryCPUHelperAbs::RunTest(const FString& Parameter)
+	{
+		FTensor XC2x1 = MakeConstTensor(TEXT("XC2x1"), { 2,1 }, { -1.0f, 2.0f });
+		FTensor Y = MakeTensor(TEXT("Y"), { 2,1 });
+		ElementWiseUnaryCPUHelper::Apply(EElementWiseUnaryOperatorType::Abs, XC2x1, 1.0f, 2.0f, 3.0f, Y);
+		UTEST_EQUAL(TEXT("Abs(XC2x1)[0]"), Y.GetPreparedData<float>()[0], FMath::Abs(XC2x1.GetPreparedData<float>()[0]));
+		UTEST_EQUAL(TEXT("Abs(XC2x1)[1]"), Y.GetPreparedData<float>()[1], FMath::Abs(XC2x1.GetPreparedData<float>()[1]));
+		return true;
+	}
+
+	IMPLEMENT_NNE_SHAPEINFERENCEHELPER_UNIT_AUTOMATION_TEST(FElementWiseUnaryCPUHelperAcos, "System.Engine.MachineLearning.NNE.UnitTest.UnaryHelper.Acos");
+	bool FElementWiseUnaryCPUHelperAcos::RunTest(const FString& Parameter)
+	{
+		FTensor XC2x1 = MakeConstTensor(TEXT("XC2x1"), { 2,1 }, { -0.5f, 0.8f });
+		FTensor Y = MakeTensor(TEXT("Y"), { 2,1 });
+		ElementWiseUnaryCPUHelper::Apply(EElementWiseUnaryOperatorType::Acos, XC2x1, 1.0f, 2.0f, 3.0f, Y);
+		UTEST_EQUAL(TEXT("Acos(XC2x1)[0]"), Y.GetPreparedData<float>()[0], FMath::Acos(XC2x1.GetPreparedData<float>()[0]));
+		UTEST_EQUAL(TEXT("Acos(XC2x1)[1]"), Y.GetPreparedData<float>()[1], FMath::Acos(XC2x1.GetPreparedData<float>()[1]));
+		return true;
+	}
+
+	IMPLEMENT_NNE_SHAPEINFERENCEHELPER_UNIT_AUTOMATION_TEST(FElementWiseUnaryCPUHelperAcosh, "System.Engine.MachineLearning.NNE.UnitTest.UnaryHelper.Acosh");
+	bool FElementWiseUnaryCPUHelperAcosh::RunTest(const FString& Parameter)
+	{
+		FTensor XC2x1 = MakeConstTensor(TEXT("XC2x1"), { 2,1 }, { 1.0f, 1.5f });
+		FTensor Y = MakeTensor(TEXT("Y"), { 2,1 });
+		ElementWiseUnaryCPUHelper::Apply(EElementWiseUnaryOperatorType::Acosh, XC2x1, 1.0f, 2.0f, 3.0f, Y);
+		UTEST_EQUAL(TEXT("Acosh(XC2x1)[0]"), Y.GetPreparedData<float>()[0], 0.0f);
+		UTEST_EQUAL(TEXT("Acosh(XC2x1)[1]"), Y.GetPreparedData<float>()[1], 0.9624236501192069f);
+		return true;
+	}
+
+	IMPLEMENT_NNE_SHAPEINFERENCEHELPER_UNIT_AUTOMATION_TEST(FElementWiseUnaryCPUHelperAsin, "System.Engine.MachineLearning.NNE.UnitTest.UnaryHelper.Asin");
+	bool FElementWiseUnaryCPUHelperAsin::RunTest(const FString& Parameter)
+	{
+		FTensor XC2x1 = MakeConstTensor(TEXT("XC2x1"), { 2,1 }, { -0.5f, 0.8f });
+		FTensor Y = MakeTensor(TEXT("Y"), { 2,1 });
+		ElementWiseUnaryCPUHelper::Apply(EElementWiseUnaryOperatorType::Asin, XC2x1, 1.0f, 2.0f, 3.0f, Y);
+		UTEST_EQUAL(TEXT("Asin(XC2x1)[0]"), Y.GetPreparedData<float>()[0], FMath::Asin(XC2x1.GetPreparedData<float>()[0]));
+		UTEST_EQUAL(TEXT("Asin(XC2x1)[1]"), Y.GetPreparedData<float>()[1], FMath::Asin(XC2x1.GetPreparedData<float>()[1]));
+		return true;
+	}
+
+	IMPLEMENT_NNE_SHAPEINFERENCEHELPER_UNIT_AUTOMATION_TEST(FElementWiseUnaryCPUHelperAsinh, "System.Engine.MachineLearning.NNE.UnitTest.UnaryHelper.Asinh");
+	bool FElementWiseUnaryCPUHelperAsinh::RunTest(const FString& Parameter)
+	{
+		FTensor XC2x1 = MakeConstTensor(TEXT("XC2x1"), { 2,1 }, { 1.0f, 2.0f });
+		FTensor Y = MakeTensor(TEXT("Y"), { 2,1 });
+		ElementWiseUnaryCPUHelper::Apply(EElementWiseUnaryOperatorType::Asinh, XC2x1, 1.0f, 2.0f, 3.0f, Y);
+		UTEST_EQUAL(TEXT("Asinh(XC2x1)[0]"), Y.GetPreparedData<float>()[0], 0.881373587019543f);
+		UTEST_EQUAL(TEXT("Asinh(XC2x1)[1]"), Y.GetPreparedData<float>()[1], 1.4436354751788103f);
+		return true;
+	}
+
+	IMPLEMENT_NNE_SHAPEINFERENCEHELPER_UNIT_AUTOMATION_TEST(FElementWiseUnaryCPUHelperAtan, "System.Engine.MachineLearning.NNE.UnitTest.UnaryHelper.Atan");
+	bool FElementWiseUnaryCPUHelperAtan::RunTest(const FString& Parameter)
+	{
+		FTensor XC2x1 = MakeConstTensor(TEXT("XC2x1"), { 2,1 }, { -0.5f, 0.8f });
+		FTensor Y = MakeTensor(TEXT("Y"), { 2,1 });
+		ElementWiseUnaryCPUHelper::Apply(EElementWiseUnaryOperatorType::Atan, XC2x1, 1.0f, 2.0f, 3.0f, Y);
+		UTEST_EQUAL(TEXT("Atan(XC2x1)[0]"), Y.GetPreparedData<float>()[0], FMath::Atan(XC2x1.GetPreparedData<float>()[0]));
+		UTEST_EQUAL(TEXT("Atan(XC2x1)[1]"), Y.GetPreparedData<float>()[1], FMath::Atan(XC2x1.GetPreparedData<float>()[1]));
+		return true;
+	}
+
+	IMPLEMENT_NNE_SHAPEINFERENCEHELPER_UNIT_AUTOMATION_TEST(FElementWiseUnaryCPUHelperCeil, "System.Engine.MachineLearning.NNE.UnitTest.UnaryHelper.Ceil");
+	bool FElementWiseUnaryCPUHelperCeil::RunTest(const FString& Parameter)
+	{
+		FTensor XC2x1 = MakeConstTensor(TEXT("XC2x1"), { 2,1 }, { -0.5f, 0.8f });
+		FTensor Y = MakeTensor(TEXT("Y"), { 2,1 });
+		ElementWiseUnaryCPUHelper::Apply(EElementWiseUnaryOperatorType::Ceil, XC2x1, 1.0f, 2.0f, 3.0f, Y);
+		UTEST_EQUAL(TEXT("Ceil(XC2x1)[0]"), Y.GetPreparedData<float>()[0], FMath::CeilToFloat(XC2x1.GetPreparedData<float>()[0]));
+		UTEST_EQUAL(TEXT("Ceil(XC2x1)[1]"), Y.GetPreparedData<float>()[1], FMath::CeilToFloat(XC2x1.GetPreparedData<float>()[1]));
+		return true;
+	}
+	IMPLEMENT_NNE_SHAPEINFERENCEHELPER_UNIT_AUTOMATION_TEST(FElementWiseUnaryCPUHelperCos, "System.Engine.MachineLearning.NNE.UnitTest.UnaryHelper.Cos");
+	bool FElementWiseUnaryCPUHelperCos::RunTest(const FString& Parameter)
+	{
+		FTensor XC2x1 = MakeConstTensor(TEXT("XC2x1"), { 2,1 }, { -0.5f, 0.8f });
+		FTensor Y = MakeTensor(TEXT("Y"), { 2,1 });
+		ElementWiseUnaryCPUHelper::Apply(EElementWiseUnaryOperatorType::Cos, XC2x1, 1.0f, 2.0f, 3.0f, Y);
+		UTEST_EQUAL(TEXT("Cos(XC2x1)[0]"), Y.GetPreparedData<float>()[0], FMath::Cos(XC2x1.GetPreparedData<float>()[0]));
+		UTEST_EQUAL(TEXT("Cos(XC2x1)[1]"), Y.GetPreparedData<float>()[1], FMath::Cos(XC2x1.GetPreparedData<float>()[1]));
+		return true;
+	}
+
+	IMPLEMENT_NNE_SHAPEINFERENCEHELPER_UNIT_AUTOMATION_TEST(FElementWiseUnaryCPUHelperCosh, "System.Engine.MachineLearning.NNE.UnitTest.UnaryHelper.Cosh");
+	bool FElementWiseUnaryCPUHelperCosh::RunTest(const FString& Parameter)
+	{
+		FTensor XC2x1 = MakeConstTensor(TEXT("XC2x1"), { 2,1 }, { 1.0f, 2.0f });
+		FTensor Y = MakeTensor(TEXT("Y"), { 2,1 });
+		ElementWiseUnaryCPUHelper::Apply(EElementWiseUnaryOperatorType::Cosh, XC2x1, 1.0f, 2.0f, 3.0f, Y);
+		UTEST_EQUAL(TEXT("Cosh(XC2x1)[0]"), Y.GetPreparedData<float>()[0], 1.5430806348152437f);
+		UTEST_EQUAL(TEXT("Cosh(XC2x1)[1]"), Y.GetPreparedData<float>()[1], 3.7621956910836314f);
+		return true;
+	}
+	IMPLEMENT_NNE_SHAPEINFERENCEHELPER_UNIT_AUTOMATION_TEST(FElementWiseUnaryCPUHelperElu, "System.Engine.MachineLearning.NNE.UnitTest.UnaryHelper.Elu");
+	bool FElementWiseUnaryCPUHelperElu::RunTest(const FString& Parameter)
+	{
+		FTensor XC2x1 = MakeConstTensor(TEXT("XC2x1"), { 2,1 }, { -0.5f, 0.8f });
+		FTensor Y = MakeTensor(TEXT("Y"), { 2,1 });
+		ElementWiseUnaryCPUHelper::Apply(EElementWiseUnaryOperatorType::Elu, XC2x1, 1.0f, 2.0f, 3.0f, Y);
+		UTEST_EQUAL(TEXT("Elu(XC2x1)[0]"), Y.GetPreparedData<float>()[0], -0.3934693402873666f);
+		UTEST_EQUAL(TEXT("Elu(XC2x1)[1]"), Y.GetPreparedData<float>()[1], 0.8f);
+		return true;
+	}
+	
+	IMPLEMENT_NNE_SHAPEINFERENCEHELPER_UNIT_AUTOMATION_TEST(FElementWiseUnaryCPUHelperErf, "System.Engine.MachineLearning.NNE.UnitTest.UnaryHelper.Erf");
+	bool FElementWiseUnaryCPUHelperErf::RunTest(const FString& Parameter)
+	{
+		FTensor XC2x1 = MakeConstTensor(TEXT("XC2x1"), { 2,1 }, { 1.0f, 2.0f });
+		FTensor Y = MakeTensor(TEXT("Y"), { 2,1 });
+		ElementWiseUnaryCPUHelper::Apply(EElementWiseUnaryOperatorType::Erf, XC2x1, 1.0f, 2.0f, 3.0f, Y);
+		UTEST_EQUAL_TOLERANCE(TEXT("Erf(XC2x1)[0]"), Y.GetPreparedData<float>()[0], 0.8427007929497149f, 1e-3f);
+		UTEST_EQUAL_TOLERANCE(TEXT("Erf(XC2x1)[1]"), Y.GetPreparedData<float>()[1], 0.9953222650189527f, 1e-3f);
+		return true;
+	}
+
+	IMPLEMENT_NNE_SHAPEINFERENCEHELPER_UNIT_AUTOMATION_TEST(FElementWiseUnaryCPUHelperExp, "System.Engine.MachineLearning.NNE.UnitTest.UnaryHelper.Exp");
+	bool FElementWiseUnaryCPUHelperExp::RunTest(const FString& Parameter)
+	{
+		FTensor XC2x1 = MakeConstTensor(TEXT("XC2x1"), { 2,1 }, { -0.5f, 0.8f });
+		FTensor Y = MakeTensor(TEXT("Y"), { 2,1 });
+		ElementWiseUnaryCPUHelper::Apply(EElementWiseUnaryOperatorType::Exp, XC2x1, 1.0f, 2.0f, 3.0f, Y);
+		UTEST_EQUAL(TEXT("Exp(XC2x1)[0]"), Y.GetPreparedData<float>()[0], FMath::Exp(XC2x1.GetPreparedData<float>()[0]));
+		UTEST_EQUAL(TEXT("Exp(XC2x1)[1]"), Y.GetPreparedData<float>()[1], FMath::Exp(XC2x1.GetPreparedData<float>()[1]));
+		return true;
+	}
+
+	IMPLEMENT_NNE_SHAPEINFERENCEHELPER_UNIT_AUTOMATION_TEST(FElementWiseUnaryCPUHelperFloor, "System.Engine.MachineLearning.NNE.UnitTest.UnaryHelper.Floor");
+	bool FElementWiseUnaryCPUHelperFloor::RunTest(const FString& Parameter)
+	{
+		FTensor XC2x1 = MakeConstTensor(TEXT("XC2x1"), { 2,1 }, { -0.5f, 0.8f });
+		FTensor Y = MakeTensor(TEXT("Y"), { 2,1 });
+		ElementWiseUnaryCPUHelper::Apply(EElementWiseUnaryOperatorType::Floor, XC2x1, 1.0f, 2.0f, 3.0f, Y);
+		UTEST_EQUAL(TEXT("Floor(XC2x1)[0]"), Y.GetPreparedData<float>()[0], FMath::Floor(XC2x1.GetPreparedData<float>()[0]));
+		UTEST_EQUAL(TEXT("Floor(XC2x1)[1]"), Y.GetPreparedData<float>()[1], FMath::Floor(XC2x1.GetPreparedData<float>()[1]));
+		return true;
+	}
+	
+	IMPLEMENT_NNE_SHAPEINFERENCEHELPER_UNIT_AUTOMATION_TEST(FElementWiseUnaryCPUHelperIsInf, "System.Engine.MachineLearning.NNE.UnitTest.UnaryHelper.IsInf");
+	bool FElementWiseUnaryCPUHelperIsInf::RunTest(const FString& Parameter)
+	{
+		FTensor XC2x1 = MakeConstTensor(TEXT("XC2x1"), { 2,1 }, { -0.5f, INFINITY });
+		FTensor Y = MakeTensor(TEXT("Y"), { 2,1 });
+		ElementWiseUnaryCPUHelper::Apply(EElementWiseUnaryOperatorType::IsInf, XC2x1, 1.0f, 2.0f, 3.0f, Y);
+		UTEST_EQUAL(TEXT("IsInf(XC2x1)[0]"), Y.GetPreparedData<float>()[0], 0.0f);
+		UTEST_EQUAL(TEXT("IsInf(XC2x1)[1]"), Y.GetPreparedData<float>()[1], 1.0f);
+		return true;
+	}
+
+	IMPLEMENT_NNE_SHAPEINFERENCEHELPER_UNIT_AUTOMATION_TEST(FElementWiseUnaryCPUHelperIsNan, "System.Engine.MachineLearning.NNE.UnitTest.UnaryHelper.IsNan");
+	bool FElementWiseUnaryCPUHelperIsNan::RunTest(const FString& Parameter)
+	{
+		float FloatNan = FMath::Sqrt(-1.0f);
+		FTensor XC2x1 = MakeConstTensor(TEXT("XC2x1"), { 2,1 }, { -0.5f, FloatNan });
+		FTensor Y = MakeTensor(TEXT("Y"), { 2,1 });
+		ElementWiseUnaryCPUHelper::Apply(EElementWiseUnaryOperatorType::IsNan, XC2x1, 1.0f, 2.0f, 3.0f, Y);
+		UTEST_EQUAL(TEXT("IsNan(XC2x1)[0]"), Y.GetPreparedData<float>()[0], 0.0f);
+		UTEST_EQUAL(TEXT("IsNan(XC2x1)[1]"), Y.GetPreparedData<float>()[1], 1.0f);
+		return true;
+	}
+
+	IMPLEMENT_NNE_SHAPEINFERENCEHELPER_UNIT_AUTOMATION_TEST(FElementWiseUnaryCPUHelperHardSigmoid, "System.Engine.MachineLearning.NNE.UnitTest.UnaryHelper.HardSigmoid");
+	bool FElementWiseUnaryCPUHelperHardSigmoid::RunTest(const FString& Parameter)
+	{
+		FTensor XC2x1 = MakeConstTensor(TEXT("XC2x1"), { 2,1 }, { 0.0f, -1.5f });
+		FTensor Y = MakeTensor(TEXT("Y"), { 2,1 });
+		ElementWiseUnaryCPUHelper::Apply(EElementWiseUnaryOperatorType::HardSigmoid, XC2x1, 1.0f, 2.0f, 3.0f, Y);
+		UTEST_EQUAL(TEXT("HardSigmoid(XC2x1)[0]"), Y.GetPreparedData<float>()[0], 1.0f);
+		UTEST_EQUAL(TEXT("HardSigmoid(XC2x1)[1]"), Y.GetPreparedData<float>()[1], 0.5f);
+		return true;
+	}
+
+	IMPLEMENT_NNE_SHAPEINFERENCEHELPER_UNIT_AUTOMATION_TEST(FElementWiseUnaryCPUHelperHardswitch, "System.Engine.MachineLearning.NNE.UnitTest.UnaryHelper.HardSwitch");
+	bool FElementWiseUnaryCPUHelperHardswitch::RunTest(const FString& Parameter)
+	{
+		FTensor XC2x1 = MakeConstTensor(TEXT("XC2x1"), { 2,1 }, { 0.0f, 1.0f });
+		FTensor Y = MakeTensor(TEXT("Y"), { 2,1 });
+		ElementWiseUnaryCPUHelper::Apply(EElementWiseUnaryOperatorType::HardSwish, XC2x1, 1.0f, 2.0f, 3.0f, Y);
+		UTEST_EQUAL(TEXT("HardSwitch(XC2x1)[0]"), Y.GetPreparedData<float>()[0], 0.5f);
+		UTEST_EQUAL(TEXT("HardSwitch(XC2x1)[1]"), Y.GetPreparedData<float>()[1], (1.0f / 6.0f) + 0.5f);
+		return true;
+	}
+
+	IMPLEMENT_NNE_SHAPEINFERENCEHELPER_UNIT_AUTOMATION_TEST(FElementWiseUnaryCPUHelperLeakyRelu, "System.Engine.MachineLearning.NNE.UnitTest.UnaryHelper.LeakyRelu");
+	bool FElementWiseUnaryCPUHelperLeakyRelu::RunTest(const FString& Parameter)
+	{
+		FTensor XC2x1 = MakeConstTensor(TEXT("XC2x1"), { 2,1 }, { 1.0f, -1.0f });
+		FTensor Y = MakeTensor(TEXT("Y"), { 2,1 });
+		ElementWiseUnaryCPUHelper::Apply(EElementWiseUnaryOperatorType::LeakyRelu, XC2x1, -0.5f, 2.0f, 3.0f, Y);
+		UTEST_EQUAL(TEXT("LeakyRelu(XC2x1)[0]"), Y.GetPreparedData<float>()[0], 1.0f);
+		UTEST_EQUAL(TEXT("LeakyRelu(XC2x1)[1]"), Y.GetPreparedData<float>()[1], 0.5f);
+		return true;
+	}
+
+	IMPLEMENT_NNE_SHAPEINFERENCEHELPER_UNIT_AUTOMATION_TEST(FElementWiseUnaryCPUHelperLog, "System.Engine.MachineLearning.NNE.UnitTest.UnaryHelper.Log");
+	bool FElementWiseUnaryCPUHelperLog::RunTest(const FString& Parameter)
+	{
+		FTensor XC2x1 = MakeConstTensor(TEXT("XC2x1"), { 2,1 }, { -0.5f, 0.8f });
+		FTensor Y = MakeTensor(TEXT("Y"), { 2,1 });
+		ElementWiseUnaryCPUHelper::Apply(EElementWiseUnaryOperatorType::Log, XC2x1, 1.0f, 2.0f, 3.0f, Y);
+		UTEST_EQUAL(TEXT("Log(XC2x1)[0]"), Y.GetPreparedData<float>()[0], FMath::Loge(XC2x1.GetPreparedData<float>()[0]));
+		UTEST_EQUAL(TEXT("Log(XC2x1)[1]"), Y.GetPreparedData<float>()[1], FMath::Loge(XC2x1.GetPreparedData<float>()[1]));
+		return true;
+	}
+
+	IMPLEMENT_NNE_SHAPEINFERENCEHELPER_UNIT_AUTOMATION_TEST(FElementWiseUnaryCPUHelperNeg, "System.Engine.MachineLearning.NNE.UnitTest.UnaryHelper.Neg");
+	bool FElementWiseUnaryCPUHelperNeg::RunTest(const FString& Parameter)
+	{
+		FTensor XC2x1 = MakeConstTensor(TEXT("XC2x1"), { 2,1 }, { -0.5f, 0.8f });
+		FTensor Y = MakeTensor(TEXT("Y"), { 2,1 });
+		ElementWiseUnaryCPUHelper::Apply(EElementWiseUnaryOperatorType::Neg, XC2x1, 1.0f, 2.0f, 3.0f, Y);
+		UTEST_EQUAL(TEXT("Neg(XC2x1)[0]"), Y.GetPreparedData<float>()[0], -XC2x1.GetPreparedData<float>()[0]);
+		UTEST_EQUAL(TEXT("Neg(XC2x1)[1]"), Y.GetPreparedData<float>()[1], -XC2x1.GetPreparedData<float>()[1]);
+		return true;
+	}
+
+	IMPLEMENT_NNE_SHAPEINFERENCEHELPER_UNIT_AUTOMATION_TEST(FElementWiseUnaryCPUHelperReciprocal, "System.Engine.MachineLearning.NNE.UnitTest.UnaryHelper.Reciprocal");
+	bool FElementWiseUnaryCPUHelperReciprocal::RunTest(const FString& Parameter)
+	{
+		FTensor XC2x1 = MakeConstTensor(TEXT("XC2x1"), { 2,1 }, { -0.5f, 0.8f });
+		FTensor Y = MakeTensor(TEXT("Y"), { 2,1 });
+		ElementWiseUnaryCPUHelper::Apply(EElementWiseUnaryOperatorType::Reciprocal, XC2x1, 1.0f, 2.0f, 3.0f, Y);
+		UTEST_EQUAL(TEXT("Reciprocal(XC2x1)[0]"), Y.GetPreparedData<float>()[0], 1.0f / (XC2x1.GetPreparedData<float>()[0]));
+		UTEST_EQUAL(TEXT("Reciprocal(XC2x1)[1]"), Y.GetPreparedData<float>()[1], 1.0f / (XC2x1.GetPreparedData<float>()[1]));
+		return true;
+	}
+
+	IMPLEMENT_NNE_SHAPEINFERENCEHELPER_UNIT_AUTOMATION_TEST(FElementWiseUnaryCPUHelperRelu, "System.Engine.MachineLearning.NNE.UnitTest.UnaryHelper.Relu");
+	bool FElementWiseUnaryCPUHelperRelu::RunTest(const FString& Parameter)
+	{
+		FTensor XC2x1 = MakeConstTensor(TEXT("XC2x1"), { 2,1 }, { -0.5f, 0.8f });
+		FTensor Y = MakeTensor(TEXT("Y"), { 2,1 });
+		ElementWiseUnaryCPUHelper::Apply(EElementWiseUnaryOperatorType::Relu, XC2x1, 1.0f, 2.0f, 3.0f, Y);
+		UTEST_EQUAL(TEXT("Relu(XC2x1)[0]"), Y.GetPreparedData<float>()[0], 0.0f);
+		UTEST_EQUAL(TEXT("Relu(XC2x1)[1]"), Y.GetPreparedData<float>()[1], 0.8f);
+		return true;
+	}
+
+	IMPLEMENT_NNE_SHAPEINFERENCEHELPER_UNIT_AUTOMATION_TEST(FElementWiseUnaryCPUHelperSelu, "System.Engine.MachineLearning.NNE.UnitTest.UnaryHelper.Selu");
+	bool FElementWiseUnaryCPUHelperSelu::RunTest(const FString& Parameter)
+	{
+		FTensor XC2x1 = MakeConstTensor(TEXT("XC2x1"), { 2,1 }, { 0.1f, -1.0f });
+		FTensor Y = MakeTensor(TEXT("Y"), { 2,1 });
+		ElementWiseUnaryCPUHelper::Apply(EElementWiseUnaryOperatorType::Selu, XC2x1, 1.0f, 2.0f, 3.0f, Y);
+		UTEST_EQUAL(TEXT("Selu(XC2x1)[0]"), Y.GetPreparedData<float>()[0], 0.3f);
+		UTEST_EQUAL(TEXT("Selu(XC2x1)[1]"), Y.GetPreparedData<float>()[1], 3.0f * (0.36787944117144233f - 1.0f));
+		return true;
+	}
+
+	IMPLEMENT_NNE_SHAPEINFERENCEHELPER_UNIT_AUTOMATION_TEST(FElementWiseUnaryCPUHelperSigmoid, "System.Engine.MachineLearning.NNE.UnitTest.UnaryHelper.Sigmoid");
+	bool FElementWiseUnaryCPUHelperSigmoid::RunTest(const FString& Parameter)
+	{
+		FTensor XC2x1 = MakeConstTensor(TEXT("XC2x1"), { 2,1 }, { 0.0f, 1.0f });
+		FTensor Y = MakeTensor(TEXT("Y"), { 2,1 });
+		ElementWiseUnaryCPUHelper::Apply(EElementWiseUnaryOperatorType::Sigmoid, XC2x1, 1.0f, 2.0f, 3.0f, Y);
+		UTEST_EQUAL(TEXT("Sigmoid(XC2x1)[0]"), Y.GetPreparedData<float>()[0], 0.5f);
+		UTEST_EQUAL(TEXT("Sigmoid(XC2x1)[1]"), Y.GetPreparedData<float>()[1], 1.0f / (1.0f + 0.36787944117144233f));
+		return true;
+	}
+
+	IMPLEMENT_NNE_SHAPEINFERENCEHELPER_UNIT_AUTOMATION_TEST(FElementWiseUnaryCPUHelperSign, "System.Engine.MachineLearning.NNE.UnitTest.UnaryHelper.Sign");
+	bool FElementWiseUnaryCPUHelperSign::RunTest(const FString& Parameter)
+	{
+		FTensor XC2x1 = MakeConstTensor(TEXT("XC2x1"), { 2,1 }, { -0.5f, 0.8f });
+		FTensor Y = MakeTensor(TEXT("Y"), { 2,1 });
+		ElementWiseUnaryCPUHelper::Apply(EElementWiseUnaryOperatorType::Sign, XC2x1, 1.0f, 2.0f, 3.0f, Y);
+		UTEST_EQUAL(TEXT("Sign(XC2x1)[0]"), Y.GetPreparedData<float>()[0], FMath::Sign(XC2x1.GetPreparedData<float>()[0]));
+		UTEST_EQUAL(TEXT("Sign(XC2x1)[1]"), Y.GetPreparedData<float>()[1], FMath::Sign(XC2x1.GetPreparedData<float>()[1]));
+		return true;
+	}
+
+	IMPLEMENT_NNE_SHAPEINFERENCEHELPER_UNIT_AUTOMATION_TEST(FElementWiseUnaryCPUHelperSin, "System.Engine.MachineLearning.NNE.UnitTest.UnaryHelper.Sin");
+	bool FElementWiseUnaryCPUHelperSin::RunTest(const FString& Parameter)
+	{
+		FTensor XC2x1 = MakeConstTensor(TEXT("XC2x1"), { 2,1 }, { -0.5f, 0.8f });
+		FTensor Y = MakeTensor(TEXT("Y"), { 2,1 });
+		ElementWiseUnaryCPUHelper::Apply(EElementWiseUnaryOperatorType::Sin, XC2x1, 1.0f, 2.0f, 3.0f, Y);
+		UTEST_EQUAL(TEXT("Sin(XC2x1)[0]"), Y.GetPreparedData<float>()[0], FMath::Sin(XC2x1.GetPreparedData<float>()[0]));
+		UTEST_EQUAL(TEXT("Sin(XC2x1)[1]"), Y.GetPreparedData<float>()[1], FMath::Sin(XC2x1.GetPreparedData<float>()[1]));
+		return true;
+	}
+
+	IMPLEMENT_NNE_SHAPEINFERENCEHELPER_UNIT_AUTOMATION_TEST(FElementWiseUnaryCPUHelperSinh, "System.Engine.MachineLearning.NNE.UnitTest.UnaryHelper.Sinh");
+	bool FElementWiseUnaryCPUHelperSinh::RunTest(const FString& Parameter)
+	{
+		FTensor XC2x1 = MakeConstTensor(TEXT("XC2x1"), { 2,1 }, { -0.5f, 0.8f });
+		FTensor Y = MakeTensor(TEXT("Y"), { 2,1 });
+		ElementWiseUnaryCPUHelper::Apply(EElementWiseUnaryOperatorType::Sinh, XC2x1, 1.0f, 2.0f, 3.0f, Y);
+		UTEST_EQUAL(TEXT("Sinh(XC2x1)[0]"), Y.GetPreparedData<float>()[0], FMath::Sinh(XC2x1.GetPreparedData<float>()[0]));
+		UTEST_EQUAL(TEXT("Sinh(XC2x1)[1]"), Y.GetPreparedData<float>()[1], FMath::Sinh(XC2x1.GetPreparedData<float>()[1]));
+		return true;
+	}
+
+	IMPLEMENT_NNE_SHAPEINFERENCEHELPER_UNIT_AUTOMATION_TEST(FElementWiseUnaryCPUHelperSoftplus, "System.Engine.MachineLearning.NNE.UnitTest.UnaryHelper.Softplus");
+	bool FElementWiseUnaryCPUHelperSoftplus::RunTest(const FString& Parameter)
+	{
+		FTensor XC2x1 = MakeConstTensor(TEXT("XC2x1"), { 2,1 }, { 0.0f, 1.0f });
+		FTensor Y = MakeTensor(TEXT("Y"), { 2,1 });
+		ElementWiseUnaryCPUHelper::Apply(EElementWiseUnaryOperatorType::Softplus, XC2x1, 1.0f, 2.0f, 3.0f, Y);
+		UTEST_EQUAL(TEXT("Softplus(XC2x1)[0]"), Y.GetPreparedData<float>()[0], 0.6931471805599453f);
+		UTEST_EQUAL(TEXT("Softplus(XC2x1)[1]"), Y.GetPreparedData<float>()[1], 1.3132616875182228f);
+		return true;
+	}
+
+	IMPLEMENT_NNE_SHAPEINFERENCEHELPER_UNIT_AUTOMATION_TEST(FElementWiseUnaryCPUHelperSoftsign, "System.Engine.MachineLearning.NNE.UnitTest.UnaryHelper.Softsign");
+	bool FElementWiseUnaryCPUHelperSoftsign::RunTest(const FString& Parameter)
+	{
+		FTensor XC2x1 = MakeConstTensor(TEXT("XC2x1"), { 2,1 }, { 0.0f, 1.0f });
+		FTensor Y = MakeTensor(TEXT("Y"), { 2,1 });
+		ElementWiseUnaryCPUHelper::Apply(EElementWiseUnaryOperatorType::Softsign, XC2x1, 1.0f, 2.0f, 3.0f, Y);
+		UTEST_EQUAL(TEXT("Softsign(XC2x1)[0]"), Y.GetPreparedData<float>()[0], 0.0f);
+		UTEST_EQUAL(TEXT("Softsign(XC2x1)[1]"), Y.GetPreparedData<float>()[1], 0.5f);
+		return true;
+	}
+
+	IMPLEMENT_NNE_SHAPEINFERENCEHELPER_UNIT_AUTOMATION_TEST(FElementWiseUnaryCPUHelperSqrt, "System.Engine.MachineLearning.NNE.UnitTest.UnaryHelper.Sqrt");
+	bool FElementWiseUnaryCPUHelperSqrt::RunTest(const FString& Parameter)
+	{
+		FTensor XC2x1 = MakeConstTensor(TEXT("XC2x1"), { 2,1 }, { 2.0f, 1.0f });
+		FTensor Y = MakeTensor(TEXT("Y"), { 2,1 });
+		ElementWiseUnaryCPUHelper::Apply(EElementWiseUnaryOperatorType::Sqrt, XC2x1, 1.0f, 2.0f, 3.0f, Y);
+		UTEST_EQUAL(TEXT("Sqrt(XC2x1)[0]"), Y.GetPreparedData<float>()[0], FMath::Sqrt(XC2x1.GetPreparedData<float>()[0]));
+		UTEST_EQUAL(TEXT("Sqrt(XC2x1)[1]"), Y.GetPreparedData<float>()[1], FMath::Sqrt(XC2x1.GetPreparedData<float>()[1]));
+		return true;
+	}
+
+	IMPLEMENT_NNE_SHAPEINFERENCEHELPER_UNIT_AUTOMATION_TEST(FElementWiseUnaryCPUHelperTan, "System.Engine.MachineLearning.NNE.UnitTest.UnaryHelper.Tan");
+	bool FElementWiseUnaryCPUHelperTan::RunTest(const FString& Parameter)
+	{
+		FTensor XC2x1 = MakeConstTensor(TEXT("XC2x1"), { 2,1 }, { -0.5f, 0.8f });
+		FTensor Y = MakeTensor(TEXT("Y"), { 2,1 });
+		ElementWiseUnaryCPUHelper::Apply(EElementWiseUnaryOperatorType::Tan, XC2x1, 1.0f, 2.0f, 3.0f, Y);
+		UTEST_EQUAL(TEXT("Tan(XC2x1)[0]"), Y.GetPreparedData<float>()[0], FMath::Tan(XC2x1.GetPreparedData<float>()[0]));
+		UTEST_EQUAL(TEXT("Tan(XC2x1)[1]"), Y.GetPreparedData<float>()[1], FMath::Tan(XC2x1.GetPreparedData<float>()[1]));
+		return true;
+	}
+
+	IMPLEMENT_NNE_SHAPEINFERENCEHELPER_UNIT_AUTOMATION_TEST(FElementWiseUnaryCPUHelperTanh, "System.Engine.MachineLearning.NNE.UnitTest.UnaryHelper.Tanh");
+	bool FElementWiseUnaryCPUHelperTanh::RunTest(const FString& Parameter)
+	{
+		FTensor XC2x1 = MakeConstTensor(TEXT("XC2x1"), { 2,1 }, { 1.0f, 2.0f });
+		FTensor Y = MakeTensor(TEXT("Y"), { 2,1 });
+		ElementWiseUnaryCPUHelper::Apply(EElementWiseUnaryOperatorType::Tanh, XC2x1, 1.0f, 2.0f, 3.0f, Y);
+		UTEST_EQUAL(TEXT("Tanh(XC2x1)[0]"), Y.GetPreparedData<float>()[0], 0.7615941559557649f);
+		UTEST_EQUAL(TEXT("Tanh(XC2x1)[1]"), Y.GetPreparedData<float>()[1], 0.9640275800758169f);
+		return true;
+	}
+
+#endif //WITH_DEV_AUTOMATION_TESTS
+
+} // namespace UE::NNEQA::Private::NNERuntimeRDG::UnaryOp
diff --git a/Engine/Plugins/Experimental/NNE/Source/NNERuntimeRDG/Internal/NNERuntimeRDGElementWiseBinaryHelper.h b/Engine/Plugins/Experimental/NNE/Source/NNERuntimeRDG/Internal/NNERuntimeRDGElementWiseBinaryHelper.h
new file mode 100644
--- /dev/null
+++ b/Engine/Plugins/Experimental/NNE/Source/NNERuntimeRDG/Internal/NNERuntimeRDGElementWiseBinaryHelper.h
@@ -0,0 +1,14 @@
+// Copyright Epic Games, Inc. All Rights Reserved.
+
+#pragma once 
+
+#include "CoreMinimal.h"
+#include "NNECoreOperator.h"
+
+namespace UE::NNECore::Internal { class FTensor; }
+
+namespace UE::NNERuntimeRDG::Internal::ElementWiseBinaryCPUHelper
+{
+	void NNERUNTIMERDG_API Apply(NNECore::Internal::EElementWiseBinaryOperatorType OpType, const NNECore::Internal::FTensor& LHSTensor, const NNECore::Internal::FTensor& RHSTensor, NNECore::Internal::FTensor& OutputTensor);
+	
+} // UE::NNERuntimeRDG::Internal::ElementWiseBinaryCPUHelper
\ No newline at end of file
diff --git a/Engine/Plugins/Experimental/NNE/Source/NNERuntimeRDG/Internal/NNERuntimeRDGElementWiseUnaryHelper.h b/Engine/Plugins/Experimental/NNE/Source/NNERuntimeRDG/Internal/NNERuntimeRDGElementWiseUnaryHelper.h
new file mode 100644
--- /dev/null
+++ b/Engine/Plugins/Experimental/NNE/Source/NNERuntimeRDG/Internal/NNERuntimeRDGElementWiseUnaryHelper.h
@@ -0,0 +1,13 @@
+// Copyright Epic Games, Inc. All Rights Reserved.
+
+#pragma once 
+
+#include "CoreMinimal.h"
+#include "NNECoreOperator.h"
+
+namespace UE::NNECore::Internal { class FTensor; }
+
+namespace UE::NNERuntimeRDG::Internal::ElementWiseUnaryCPUHelper
+{
+	void NNERUNTIMERDG_API Apply(NNECore::Internal::EElementWiseUnaryOperatorType OpType, const NNECore::Internal::FTensor& Tensor, float Alpha, float Beta, float Gamma, NNECore::Internal::FTensor& OutputTensor);
+} // UE::NNERuntimeRDG::Internal::ElementWiseUnaryCPUHelper
\ No newline at end of file
diff --git a/Engine/Plugins/Experimental/NNE/Source/NNERuntimeRDG/Private/Hlsl/NNERuntimeRDGElementWiseBinary.cpp b/Engine/Plugins/Experimental/NNE/Source/NNERuntimeRDG/Private/Hlsl/NNERuntimeRDGElementWiseBinary.cpp
--- a/Engine/Plugins/Experimental/NNE/Source/NNERuntimeRDG/Private/Hlsl/NNERuntimeRDGElementWiseBinary.cpp
+++ b/Engine/Plugins/Experimental/NNE/Source/NNERuntimeRDG/Private/Hlsl/NNERuntimeRDGElementWiseBinary.cpp
@@ -1,6 +1,7 @@
 // Copyright Epic Games, Inc. All Rights Reserved.
 
 #include "NNERuntimeRDGElementWiseBinary.h"
+#include "NNERuntimeRDGElementWiseBinaryHelper.h"
 #include "NNEHlslShadersElementWiseBinaryCS.h"
 #include "NNERuntimeRDGHlslHelper.h"
 #include "NNECoreTensor.h"
@@ -56,6 +57,8 @@ namespace UE::NNERuntimeRDG::Private::Hlsl
 			NNECore::FTensorShape OutputShape = NNECore::FTensorShape::Make(OutputShapeData);
 			
 			OutputTensors[0]->SetShape(OutputShape);
+
+			Internal::ElementWiseBinaryCPUHelper::Apply(OpType, *InputTensors[0], *InputTensors[1], *OutputTensors[0]);
 			
 			return 0;
 		}
diff --git a/Engine/Plugins/Experimental/NNE/Source/NNERuntimeRDG/Private/Hlsl/NNERuntimeRDGElementWiseUnary.cpp b/Engine/Plugins/Experimental/NNE/Source/NNERuntimeRDG/Private/Hlsl/NNERuntimeRDGElementWiseUnary.cpp
--- a/Engine/Plugins/Experimental/NNE/Source/NNERuntimeRDG/Private/Hlsl/NNERuntimeRDGElementWiseUnary.cpp
+++ b/Engine/Plugins/Experimental/NNE/Source/NNERuntimeRDG/Private/Hlsl/NNERuntimeRDGElementWiseUnary.cpp
@@ -41,17 +41,8 @@ namespace UE::NNERuntimeRDG::Private::Hlsl
 			OutputTensors[0]->SetShape(InputTensors[0]->GetShape());
 
 			const NNECore::Internal::FTensor& X = *InputTensors[0];
-			if (X.HasPreparedData())
-			{
-				TConstArrayView<float> XData = X.GetPreparedData<float>();
-				TArray<float> OutputData;
-				OutputData.Reserve(XData.Num());
-				for (float elem : XData)
-				{
-					OutputData.Add(ElementWiseUnaryCPUHelper::Apply<OpType>(elem, Alpha, Beta, Gamma));
-				}
-				OutputTensors[0]->SetPreparedData<float>(OutputData);
-			}
+
+			Internal::ElementWiseUnaryCPUHelper::Apply(OpType, X, Alpha, Beta, Gamma, *OutputTensors[0]);
 			
 			return 0;
 		}
diff --git a/Engine/Plugins/Experimental/NNE/Source/NNERuntimeRDG/Private/Hlsl/NNERuntimeRDGShape.cpp b/Engine/Plugins/Experimental/NNE/Source/NNERuntimeRDG/Private/Hlsl/NNERuntimeRDGShape.cpp
--- a/Engine/Plugins/Experimental/NNE/Source/NNERuntimeRDG/Private/Hlsl/NNERuntimeRDGShape.cpp
+++ b/Engine/Plugins/Experimental/NNE/Source/NNERuntimeRDG/Private/Hlsl/NNERuntimeRDGShape.cpp
@@ -74,8 +74,17 @@ namespace UE::NNERuntimeRDG::Private::Hlsl
 		bIsValid &= AttributeValidator.Validate(AttributeMap);
 
 		FInputValidator InputValidator;
+		InputValidator.AddSupportedType(ENNETensorDataType::Half);
 		InputValidator.AddSupportedType(ENNETensorDataType::Float);
+		InputValidator.AddSupportedType(ENNETensorDataType::Double);
+		InputValidator.AddSupportedType(ENNETensorDataType::Int8);
+		InputValidator.AddSupportedType(ENNETensorDataType::Int16);
+		InputValidator.AddSupportedType(ENNETensorDataType::Int32);
 		InputValidator.AddSupportedType(ENNETensorDataType::Int64);
+		InputValidator.AddSupportedType(ENNETensorDataType::UInt8);
+		InputValidator.AddSupportedType(ENNETensorDataType::UInt16);
+		InputValidator.AddSupportedType(ENNETensorDataType::UInt32);
+		InputValidator.AddSupportedType(ENNETensorDataType::UInt64);
 		InputValidator.AddRequired();
 		bIsValid &= InputValidator.Validate(InputTypes);
 
diff --git a/Engine/Plugins/Experimental/NNE/Source/NNERuntimeRDG/Private/NNERuntimeRDGElementWiseBinaryHelper.cpp b/Engine/Plugins/Experimental/NNE/Source/NNERuntimeRDG/Private/NNERuntimeRDGElementWiseBinaryHelper.cpp
new file mode 100644
--- /dev/null
+++ b/Engine/Plugins/Experimental/NNE/Source/NNERuntimeRDG/Private/NNERuntimeRDGElementWiseBinaryHelper.cpp
@@ -0,0 +1,126 @@
+// Copyright Epic Games, Inc. All Rights Reserved.
+
+#include "NNERuntimeRDGElementWiseBinaryHelper.h"
+#include "NNECoreTypes.h"
+#include "NNECoreTensor.h"
+#include "Math/UnrealMathUtility.h"
+#include "MathUtil.h"
+
+namespace UE::NNERuntimeRDG::Internal::ElementWiseBinaryCPUHelper
+{
+	class TensorIdxIterator
+	{
+		const NNECore::FTensorShape& TensorShape;
+		TArray<uint32, TInlineAllocator<NNECore::FTensorShape::MaxRank>> CurrentPosition;
+
+	public:
+		TensorIdxIterator(const NNECore::FTensorShape& InTensorShape);
+		bool Advance();
+		int32 GetIndexToBroadcastedShape(const NNECore::FTensorShape& InTensorShape) const;
+	};
+	
+	TensorIdxIterator::TensorIdxIterator(const NNECore::FTensorShape& InTensorShape) : TensorShape(InTensorShape)
+	{
+		CurrentPosition.Init(0, InTensorShape.Rank());
+	}
+
+	bool TensorIdxIterator::Advance()
+	{
+		for (int32 i = TensorShape.Rank() - 1; i >= 0; --i)
+		{
+			++CurrentPosition[i];
+			if (CurrentPosition[i] < TensorShape.GetData()[i])
+			{
+				return true;
+			}
+			CurrentPosition[i] = 0;
+		}
+		return false;
+	}
+
+	int32 TensorIdxIterator::GetIndexToBroadcastedShape(const NNECore::FTensorShape& InTensorShape) const
+	{
+		int32 Index = 0;
+		int32 DimBaseOffset = 1;
+		for (int32 r = TensorShape.Rank() - 1; r >= 0; --r)
+		{
+			if (r >= InTensorShape.Rank())
+			{
+				break;
+			}
+			Index += FMath::Min(CurrentPosition[r], InTensorShape.GetData()[r]-1) * DimBaseOffset;
+			DimBaseOffset *= FMath::Min(TensorShape.GetData()[r], InTensorShape.GetData()[r]);
+		}
+		return Index;
+	}
+
+	template<NNECore::Internal::EElementWiseBinaryOperatorType OpType> float Apply(float X, float Y);
+	template<> float Apply<NNECore::Internal::EElementWiseBinaryOperatorType::Add>(float X, float Y) { return X + Y; }
+	template<> float Apply<NNECore::Internal::EElementWiseBinaryOperatorType::Div>(float X, float Y) { return X / Y; }
+	template<> float Apply<NNECore::Internal::EElementWiseBinaryOperatorType::Mod>(float X, float Y) { return FMath::Fmod(X, Y); }
+	template<> float Apply<NNECore::Internal::EElementWiseBinaryOperatorType::Mul>(float X, float Y) { return X * Y; }
+	template<> float Apply<NNECore::Internal::EElementWiseBinaryOperatorType::Prelu>(float X, float Y) { return (X < 0.0f) ? (Y * X) : X; }
+	template<> float Apply<NNECore::Internal::EElementWiseBinaryOperatorType::Pow>(float X, float Y) { return FMath::Pow(X, Y); }
+	template<> float Apply<NNECore::Internal::EElementWiseBinaryOperatorType::Sub>(float X, float Y) { return X - Y; }
+
+	template<NNECore::Internal::EElementWiseBinaryOperatorType OpType> void Apply(const NNECore::Internal::FTensor& LHSTensor, const NNECore::Internal::FTensor& RHSTensor, NNECore::Internal::FTensor& OutputTensor)
+	{
+		//Heuristic to avoid unexpected performance hit. This helper being intended for shape related arithmetic only.
+		static constexpr int32 MaxItemInInputTensors = NNECore::FTensorShape::MaxRank * 2;
+
+		if (LHSTensor.HasPreparedData() && 
+			RHSTensor.HasPreparedData() && 
+			(LHSTensor.GetVolume() <= MaxItemInInputTensors) &&
+			(RHSTensor.GetVolume() <= MaxItemInInputTensors))
+		{
+			TConstArrayView<float> LHSData = LHSTensor.GetPreparedData<float>();
+			TConstArrayView<float> RHSData = RHSTensor.GetPreparedData<float>();
+			TArray<float> OutputData;
+			OutputData.Reserve(OutputTensor.GetVolume());
+
+			TensorIdxIterator it(OutputTensor.GetShape());
+			do
+			{
+				int32 LHSIdx = it.GetIndexToBroadcastedShape(LHSTensor.GetShape());
+				int32 RHSIdx = it.GetIndexToBroadcastedShape(RHSTensor.GetShape());
+				float LHSValue = LHSData.GetData()[LHSIdx];
+				float RHSValue = RHSData.GetData()[RHSIdx];
+				OutputData.Add(Apply<OpType>(LHSValue, RHSValue));
+			} while (it.Advance());
+
+			check(OutputData.Num() == OutputTensor.GetVolume());
+			OutputTensor.SetPreparedData<float>(OutputData);
+		}
+	}
+
+	void Apply(NNECore::Internal::EElementWiseBinaryOperatorType OpType, const NNECore::Internal::FTensor& LHSTensor, const NNECore::Internal::FTensor& RHSTensor, NNECore::Internal::FTensor& OutputTensor)
+	{
+		switch (OpType)
+		{
+		case NNECore::Internal::EElementWiseBinaryOperatorType::Add:
+			Apply<NNECore::Internal::EElementWiseBinaryOperatorType::Add>(LHSTensor, RHSTensor, OutputTensor);
+			break;
+		case NNECore::Internal::EElementWiseBinaryOperatorType::Div:
+			Apply<NNECore::Internal::EElementWiseBinaryOperatorType::Div>(LHSTensor, RHSTensor, OutputTensor);
+			break;
+		case NNECore::Internal::EElementWiseBinaryOperatorType::Mod:
+			Apply<NNECore::Internal::EElementWiseBinaryOperatorType::Mod>(LHSTensor, RHSTensor, OutputTensor);
+			break;
+		case NNECore::Internal::EElementWiseBinaryOperatorType::Mul:
+			Apply<NNECore::Internal::EElementWiseBinaryOperatorType::Mul>(LHSTensor, RHSTensor, OutputTensor);
+			break;
+		case NNECore::Internal::EElementWiseBinaryOperatorType::Prelu:
+			Apply<NNECore::Internal::EElementWiseBinaryOperatorType::Prelu>(LHSTensor, RHSTensor, OutputTensor);
+			break;
+		case NNECore::Internal::EElementWiseBinaryOperatorType::Pow:
+			Apply<NNECore::Internal::EElementWiseBinaryOperatorType::Pow>(LHSTensor, RHSTensor, OutputTensor);
+			break;
+		case NNECore::Internal::EElementWiseBinaryOperatorType::Sub:
+			Apply<NNECore::Internal::EElementWiseBinaryOperatorType::Sub>(LHSTensor, RHSTensor, OutputTensor);
+			break;
+		default:
+			break;
+		}
+	}
+	
+} // UE::NNERuntimeRDG::Internal::ElementWiseUnaryCPUHelper
diff --git a/Engine/Plugins/Experimental/NNE/Source/NNERuntimeRDG/Private/NNERuntimeRDGElementWiseUnaryHelper.cpp b/Engine/Plugins/Experimental/NNE/Source/NNERuntimeRDG/Private/NNERuntimeRDGElementWiseUnaryHelper.cpp
--- a/Engine/Plugins/Experimental/NNE/Source/NNERuntimeRDG/Private/NNERuntimeRDGElementWiseUnaryHelper.cpp
+++ b/Engine/Plugins/Experimental/NNE/Source/NNERuntimeRDG/Private/NNERuntimeRDGElementWiseUnaryHelper.cpp
@@ -1,10 +1,13 @@
 // Copyright Epic Games, Inc. All Rights Reserved.
 
 #include "NNERuntimeRDGElementWiseUnaryHelper.h"
+#include "NNECoreTensor.h"
 #include "Math/UnrealMathUtility.h"
 
-namespace UE::NNERuntimeRDG::Private::ElementWiseUnaryCPUHelper
+namespace UE::NNERuntimeRDG::Internal::ElementWiseUnaryCPUHelper
 {
+	template<NNECore::Internal::EElementWiseUnaryOperatorType OpType> float Apply(float X, float Alpha, float Beta, float Gamma);
+	
 	template<> float Apply<NNECore::Internal::EElementWiseUnaryOperatorType::Abs>(float X, float Alpha, float Beta, float Gamma) { return FMath::Abs(X); }
 
 	template<> float Apply<NNECore::Internal::EElementWiseUnaryOperatorType::Acos>(float X, float Alpha, float Beta, float Gamma) { return FMath::Acos(X); }
@@ -12,7 +15,7 @@ namespace UE::NNERuntimeRDG::Private::ElementWiseUnaryCPUHelper
 	template<> float Apply<NNECore::Internal::EElementWiseUnaryOperatorType::Acosh>(float X, float Alpha, float Beta, float Gamma) {
 		//https://mathworld.wolfram.com/InverseHyperbolicCosine.html
 		float FloatNan = FMath::Sqrt(-1.0f);
-		float yAboveOne = FMath::Loge(X + FMath::Sqrt(X + 1.0f) + FMath::Sqrt(X - 1.0f));
+		float yAboveOne = FMath::Loge(X + FMath::Sqrt(X + 1.0f) * FMath::Sqrt(X - 1.0f));
 		if (X == 1.0f)
 		{
 			return 0.0f;
@@ -43,7 +46,7 @@ namespace UE::NNERuntimeRDG::Private::ElementWiseUnaryCPUHelper
 	
 	template<> float Apply<NNECore::Internal::EElementWiseUnaryOperatorType::Cosh>(float X, float Alpha, float Beta, float Gamma) {
 		//https://mathworld.wolfram.com/HyperbolicCosine.html
-		return 0.5f * (FMath::Exp(X) - FMath::Exp(-X));
+		return 0.5f * (FMath::Exp(X) + FMath::Exp(-X));
 	}
 	
 	template<> float Apply<NNECore::Internal::EElementWiseUnaryOperatorType::Elu>(float X, float Alpha, float Beta, float Gamma) {
@@ -120,5 +123,134 @@ namespace UE::NNERuntimeRDG::Private::ElementWiseUnaryCPUHelper
 		float x3 = X * X * X;
 		return Apply<NNECore::Internal::EElementWiseUnaryOperatorType::Tanh>(a * X + b * x3, Alpha, Beta, Gamma);
 	}
+
+	template<NNECore::Internal::EElementWiseUnaryOperatorType OpType> void Apply(const NNECore::Internal::FTensor& Tensor, float Alpha, float Beta, float Gamma, NNECore::Internal::FTensor& OutputTensor)
+	{
+		//Heuristic to avoid unexpected performance hit. This helper being intended for shape related arithmetic only.
+		static constexpr int32 MaxItemInInputTensors = NNECore::FTensorShape::MaxRank * 2;
+
+		if (Tensor.HasPreparedData() && (Tensor.GetVolume() <= MaxItemInInputTensors))
+		{
+			TConstArrayView<float> TensorData = Tensor.GetPreparedData<float>();
+			TArray<float> OutputData;
+			OutputData.Reserve(TensorData.Num());
+			for (float elem : TensorData)
+			{
+				OutputData.Add(Apply<OpType>(elem, Alpha, Beta, Gamma));
+			}
+			OutputTensor.SetPreparedData<float>(OutputData);
+		}
+	}
+
+	void Apply(NNECore::Internal::EElementWiseUnaryOperatorType OpType, const NNECore::Internal::FTensor& Tensor, float Alpha, float Beta, float Gamma, NNECore::Internal::FTensor& OutputTensor)
+	{
+		switch (OpType)
+		{
+		case NNECore::Internal::EElementWiseUnaryOperatorType::Abs:
+			Apply<NNECore::Internal::EElementWiseUnaryOperatorType::Abs>(Tensor, Alpha, Beta, Gamma, OutputTensor);
+			break;
+		case NNECore::Internal::EElementWiseUnaryOperatorType::Acos:
+			Apply<NNECore::Internal::EElementWiseUnaryOperatorType::Acos>(Tensor, Alpha, Beta, Gamma, OutputTensor);
+			break;
+		case NNECore::Internal::EElementWiseUnaryOperatorType::Acosh:
+			Apply<NNECore::Internal::EElementWiseUnaryOperatorType::Acosh>(Tensor, Alpha, Beta, Gamma, OutputTensor);
+			break;
+		case NNECore::Internal::EElementWiseUnaryOperatorType::Asin:
+			Apply<NNECore::Internal::EElementWiseUnaryOperatorType::Asin>(Tensor, Alpha, Beta, Gamma, OutputTensor);
+			break;
+		case NNECore::Internal::EElementWiseUnaryOperatorType::Asinh:
+			Apply<NNECore::Internal::EElementWiseUnaryOperatorType::Asinh>(Tensor, Alpha, Beta, Gamma, OutputTensor);
+			break;
+		case NNECore::Internal::EElementWiseUnaryOperatorType::Atan:
+			Apply<NNECore::Internal::EElementWiseUnaryOperatorType::Atan>(Tensor, Alpha, Beta, Gamma, OutputTensor);
+			break;
+		case NNECore::Internal::EElementWiseUnaryOperatorType::Atanh:
+			Apply<NNECore::Internal::EElementWiseUnaryOperatorType::Atanh>(Tensor, Alpha, Beta, Gamma, OutputTensor);
+			break;
+		case NNECore::Internal::EElementWiseUnaryOperatorType::Ceil:
+			Apply<NNECore::Internal::EElementWiseUnaryOperatorType::Ceil>(Tensor, Alpha, Beta, Gamma, OutputTensor);
+			break;
+		case NNECore::Internal::EElementWiseUnaryOperatorType::Cos:
+			Apply<NNECore::Internal::EElementWiseUnaryOperatorType::Cos>(Tensor, Alpha, Beta, Gamma, OutputTensor);
+			break;
+		case NNECore::Internal::EElementWiseUnaryOperatorType::Cosh:
+			Apply<NNECore::Internal::EElementWiseUnaryOperatorType::Cosh>(Tensor, Alpha, Beta, Gamma, OutputTensor);
+			break;
+		case NNECore::Internal::EElementWiseUnaryOperatorType::Elu:
+			Apply<NNECore::Internal::EElementWiseUnaryOperatorType::Elu>(Tensor, Alpha, Beta, Gamma, OutputTensor);
+			break;
+		case NNECore::Internal::EElementWiseUnaryOperatorType::Erf:
+			Apply<NNECore::Internal::EElementWiseUnaryOperatorType::Erf>(Tensor, Alpha, Beta, Gamma, OutputTensor);
+			break;
+		case NNECore::Internal::EElementWiseUnaryOperatorType::Exp:
+			Apply<NNECore::Internal::EElementWiseUnaryOperatorType::Exp>(Tensor, Alpha, Beta, Gamma, OutputTensor);
+			break;
+		case NNECore::Internal::EElementWiseUnaryOperatorType::Floor:
+			Apply<NNECore::Internal::EElementWiseUnaryOperatorType::Floor>(Tensor, Alpha, Beta, Gamma, OutputTensor);
+			break;
+		case NNECore::Internal::EElementWiseUnaryOperatorType::IsInf:
+			Apply<NNECore::Internal::EElementWiseUnaryOperatorType::IsInf>(Tensor, Alpha, Beta, Gamma, OutputTensor);
+			break;
+		case NNECore::Internal::EElementWiseUnaryOperatorType::IsNan:
+			Apply<NNECore::Internal::EElementWiseUnaryOperatorType::IsNan>(Tensor, Alpha, Beta, Gamma, OutputTensor);
+			break;
+		case NNECore::Internal::EElementWiseUnaryOperatorType::HardSigmoid:
+			Apply<NNECore::Internal::EElementWiseUnaryOperatorType::HardSigmoid>(Tensor, Alpha, Beta, Gamma, OutputTensor);
+			break;
+		case NNECore::Internal::EElementWiseUnaryOperatorType::HardSwish:
+			Apply<NNECore::Internal::EElementWiseUnaryOperatorType::HardSwish>(Tensor, Alpha, Beta, Gamma, OutputTensor);
+			break;
+		case NNECore::Internal::EElementWiseUnaryOperatorType::LeakyRelu:
+			Apply<NNECore::Internal::EElementWiseUnaryOperatorType::LeakyRelu>(Tensor, Alpha, Beta, Gamma, OutputTensor);
+			break;
+		case NNECore::Internal::EElementWiseUnaryOperatorType::Log:
+			Apply<NNECore::Internal::EElementWiseUnaryOperatorType::Log>(Tensor, Alpha, Beta, Gamma, OutputTensor);
+			break;
+		case NNECore::Internal::EElementWiseUnaryOperatorType::Neg:
+			Apply<NNECore::Internal::EElementWiseUnaryOperatorType::Neg>(Tensor, Alpha, Beta, Gamma, OutputTensor);
+			break;
+		case NNECore::Internal::EElementWiseUnaryOperatorType::Reciprocal:
+			Apply<NNECore::Internal::EElementWiseUnaryOperatorType::Reciprocal>(Tensor, Alpha, Beta, Gamma, OutputTensor);
+			break;
+		case NNECore::Internal::EElementWiseUnaryOperatorType::Relu:
+			Apply<NNECore::Internal::EElementWiseUnaryOperatorType::Relu>(Tensor, Alpha, Beta, Gamma, OutputTensor);
+			break;
+		case NNECore::Internal::EElementWiseUnaryOperatorType::Round:
+			Apply<NNECore::Internal::EElementWiseUnaryOperatorType::Round>(Tensor, Alpha, Beta, Gamma, OutputTensor);
+			break;
+		case NNECore::Internal::EElementWiseUnaryOperatorType::Selu:
+			Apply<NNECore::Internal::EElementWiseUnaryOperatorType::Selu>(Tensor, Alpha, Beta, Gamma, OutputTensor);
+			break;
+		case NNECore::Internal::EElementWiseUnaryOperatorType::Sigmoid:
+			Apply<NNECore::Internal::EElementWiseUnaryOperatorType::Sigmoid>(Tensor, Alpha, Beta, Gamma, OutputTensor);
+			break;
+		case NNECore::Internal::EElementWiseUnaryOperatorType::Sign:
+			Apply<NNECore::Internal::EElementWiseUnaryOperatorType::Sign>(Tensor, Alpha, Beta, Gamma, OutputTensor);
+			break;
+		case NNECore::Internal::EElementWiseUnaryOperatorType::Sin:
+			Apply<NNECore::Internal::EElementWiseUnaryOperatorType::Sin>(Tensor, Alpha, Beta, Gamma, OutputTensor);
+			break;
+		case NNECore::Internal::EElementWiseUnaryOperatorType::Sinh:
+			Apply<NNECore::Internal::EElementWiseUnaryOperatorType::Sinh>(Tensor, Alpha, Beta, Gamma, OutputTensor);
+			break;
+		case NNECore::Internal::EElementWiseUnaryOperatorType::Softplus:
+			Apply<NNECore::Internal::EElementWiseUnaryOperatorType::Softplus>(Tensor, Alpha, Beta, Gamma, OutputTensor);
+			break;
+		case NNECore::Internal::EElementWiseUnaryOperatorType::Softsign:
+			Apply<NNECore::Internal::EElementWiseUnaryOperatorType::Softsign>(Tensor, Alpha, Beta, Gamma, OutputTensor);
+			break;
+		case NNECore::Internal::EElementWiseUnaryOperatorType::Sqrt:
+			Apply<NNECore::Internal::EElementWiseUnaryOperatorType::Sqrt>(Tensor, Alpha, Beta, Gamma, OutputTensor);
+			break;
+		case NNECore::Internal::EElementWiseUnaryOperatorType::Tan:
+			Apply<NNECore::Internal::EElementWiseUnaryOperatorType::Tan>(Tensor, Alpha, Beta, Gamma, OutputTensor);
+			break;
+		case NNECore::Internal::EElementWiseUnaryOperatorType::Tanh:
+			Apply<NNECore::Internal::EElementWiseUnaryOperatorType::Tanh>(Tensor, Alpha, Beta, Gamma, OutputTensor);
+			break;
+		default:
+			break;
+		}
+	}
 	
 } // UE::NNERuntimeRDG::Private::ElementWiseUnaryCPUHelper
diff --git a/Engine/Plugins/Experimental/NNE/Source/NNERuntimeRDG/Private/NNERuntimeRDGElementWiseUnaryHelper.h b/Engine/Plugins/Experimental/NNE/Source/NNERuntimeRDG/Private/NNERuntimeRDGElementWiseUnaryHelper.h
deleted file mode 100644
