NNE - Add GRU to RDG DML

diff --git a/Engine/Plugins/Experimental/NNE/Source/NNEQA/Resources/NNEQATestDesc.json b/Engine/Plugins/Experimental/NNE/Source/NNEQA/Resources/NNEQATestDesc.json
--- a/Engine/Plugins/Experimental/NNE/Source/NNEQA/Resources/NNEQATestDesc.json
+++ b/Engine/Plugins/Experimental/NNE/Source/NNEQA/Resources/NNEQATestDesc.json
@@ -326,6 +326,15 @@
 			"targets": [
 				{ "target": "Upsample" }
 			]
+		},
+		{
+			"category": "GRU",
+			"runtimes": [
+				{ "name": "NNERuntimeRDGHlsl", "skip": true }
+			],
+			"targets": [
+				{ "target": "GRU", "absolute_tolerance": 1e-6 }
+			]
 		}
 	],
 
@@ -721,6 +730,30 @@
 					"comment": "from 'mosaic-9-sim.onnx'"
 				}
 			]
+		},
+		{ 
+			"name": "GRU.Defaults",
+			"datasets": [
+				{ "inputs": [ { "shape": [ 1, 3, 2 ] }, { "shape": [ 1, 15, 2 ] }, { "shape": [ 1, 15, 5 ] } ], "outputs": [ { "shape": [ 1, 1, 3, 5 ] }, { "shape": [ 1, 3, 5 ] } ] }
+			]
+		},
+		{ 
+			"name": "GRU.InitialBias",
+			"datasets": [
+				{ "inputs": [ { "shape": [ 1, 3, 3 ] }, { "shape": [ 1, 9, 3 ] }, { "shape": [ 1, 9, 3 ] }, { "shape": [ 1, 18 ] } ], "outputs": [ { "shape": [ 1, 1, 3, 3 ] }, { "shape": [ 1, 3, 3 ] } ] }
+			]
+		},
+		{ 
+			"name": "GRU.SeqLength",
+			"datasets": [
+				{ "inputs": [ { "shape": [ 2, 3, 3 ] }, { "shape": [ 1, 15, 3 ] }, { "shape": [ 1, 15, 5 ] }, { "shape": [ 1, 30 ] } ], "outputs": [ { "shape": [ 2, 1, 3, 5 ] }, { "shape": [ 1, 3, 5 ] } ] }
+			]
+		},
+		{ 
+			"name": "GRU.Activations",
+			"datasets": [
+				{ "inputs": [ { "shape": [ 2, 3, 3 ] }, { "shape": [ 1, 15, 3 ] }, { "shape": [ 1, 15, 5 ] }, { "shape": [ 1, 30 ] } ], "outputs": [ { "shape": [ 2, 1, 3, 5 ] }, { "shape": [ 1, 3, 5 ] } ] }
+			]
 		}
 	],
 
@@ -1037,6 +1070,40 @@
 					{ "name": "pads", "value": [ 0, 1, 1, 3, 0, 0, 2, 4 ], "type": "Int32Array" }
 				] }
 			]
+		},
+		{
+			"name": "GRU.Defaults",
+			"attribute_maps": [
+				{ "attributes": [
+					{ "name": "hidden_size", "value": 5, "type": "Int32" }
+				] }
+			]
+		},
+		{
+			"name": "GRU.InitialBias",
+			"attribute_maps": [
+				{ "attributes": [
+					{ "name": "hidden_size", "value": 3, "type": "Int32" }
+				] }
+			]
+		},
+		{
+			"name": "GRU.SeqLength",
+			"attribute_maps": [
+				{ "attributes": [
+					{ "name": "hidden_size", "value": 5, "type": "Int32" }
+				] }
+			]
+		},
+		{
+			"name": "GRU.Activations",
+			"attribute_maps": [
+				{ "attributes": [
+					{ "name": "hidden_size", "value": 5, "type": "Int32" },
+					{ "name": "activations", "value": ["Tanh", "Sigmoid"], "type": "StringArray" },
+					{ "name": "activation_alpha", "value": [0.01, 0.2], "type": "FloatArray" }
+				] }
+			]
 		}
 	]
 
diff --git a/Engine/Plugins/Experimental/NNE/Source/NNERuntimeRDG/Private/Dml/NNEDmlOperatorGRU.cpp b/Engine/Plugins/Experimental/NNE/Source/NNERuntimeRDG/Private/Dml/NNEDmlOperatorGRU.cpp
new file mode 100644
--- /dev/null
+++ b/Engine/Plugins/Experimental/NNE/Source/NNERuntimeRDG/Private/Dml/NNEDmlOperatorGRU.cpp
@@ -0,0 +1,317 @@
+// Copyright Epic Games, Inc. All Rights Reserved.
+
+#ifdef NNE_USE_DIRECTML
+
+#include "NNEDmlOperator.h"
+
+namespace UE::NNERuntimeRDG::Private::Dml
+{
+
+union FDMLGRUActivationOpDescUnion
+{
+	DML_ACTIVATION_RELU_OPERATOR_DESC Relu;
+	DML_ACTIVATION_LEAKY_RELU_OPERATOR_DESC LeakyRelu;
+	DML_ACTIVATION_THRESHOLDED_RELU_OPERATOR_DESC ThresholdedRelu;
+	DML_ACTIVATION_TANH_OPERATOR_DESC Tanh;
+	DML_ACTIVATION_SCALED_TANH_OPERATOR_DESC ScaledTanh;
+	DML_ACTIVATION_SIGMOID_OPERATOR_DESC Sigmoid;
+	DML_ACTIVATION_HARD_SIGMOID_OPERATOR_DESC HardSigmoid;
+    DML_ACTIVATION_ELU_OPERATOR_DESC Elu;
+    DML_ACTIVATION_SOFTSIGN_OPERATOR_DESC Softsign;
+    DML_ACTIVATION_SOFTPLUS_OPERATOR_DESC Softplus;
+};
+
+/**
+ * GRU
+ */
+class FOperatorDmlGRU : public FOperatorDml
+{
+	static NNECore::Internal::FTensor To4DTensor(const NNECore::Internal::FTensor& InTensor)
+	{
+		check(InTensor.GetShape().Rank() <= 4);
+		TArray<uint32, TInlineAllocator<4>> NewShapeArray;
+		NewShapeArray.Init(1, 4 - InTensor.GetShape().Rank());
+		NewShapeArray.Append(InTensor.GetShape().GetData());
+		
+		return NNECore::Internal::FTensor::Make(InTensor.GetName(), NNECore::FTensorShape::Make(NewShapeArray), InTensor.GetDataType());
+	}
+
+	static DML_RECURRENT_NETWORK_DIRECTION DirectionFromString(FStringView StringVal)
+	{
+		if (FCString::Stricmp(StringVal.GetData(), TEXT("FORWARD")) == 0)
+		{
+			return DML_RECURRENT_NETWORK_DIRECTION_FORWARD;
+		}
+		else if (FCString::Stricmp(StringVal.GetData(), TEXT("BACKWARD")) == 0)
+		{
+			return DML_RECURRENT_NETWORK_DIRECTION_BACKWARD;
+		}
+		else if (FCString::Stricmp(StringVal.GetData(), TEXT("BIDIRECTIONAL")) == 0)
+		{
+			return DML_RECURRENT_NETWORK_DIRECTION_BIDIRECTIONAL;
+		}
+		else
+		{
+			return DML_RECURRENT_NETWORK_DIRECTION_FORWARD;
+		}
+	}
+
+	bool InitActivations(const TArray<FString>& Activations, const TArray<float>& Alphas, 
+						 const TArray<float>& Betas, TArray<DML_OPERATOR_DESC>& ActivationOpDescs,
+						 TArray<FDMLGRUActivationOpDescUnion>& ActivationOpInnerDescs)
+	{
+		ActivationOpDescs.SetNumZeroed(Activations.Num());
+		ActivationOpInnerDescs.SetNumZeroed(Activations.Num());
+
+		for(int Idx = 0; Idx < Activations.Num(); Idx++)
+		{
+			const FString& Activation = Activations[Idx];
+            DML_OPERATOR_DESC& ActivationOpDesc = ActivationOpDescs[Idx];
+            FDMLGRUActivationOpDescUnion& ActivationOpInnerDesc = ActivationOpInnerDescs[Idx];
+			ActivationOpDesc.Desc = &ActivationOpInnerDesc;
+
+            if (Activation == TEXT("Relu"))
+            { 
+                ActivationOpDesc.Type = DML_OPERATOR_ACTIVATION_RELU;
+            }
+            else if (Activation == TEXT("LeakyRelu"))
+            { 
+                ActivationOpDesc.Type = DML_OPERATOR_ACTIVATION_LEAKY_RELU;
+                ActivationOpInnerDesc.LeakyRelu.Alpha = Idx >= Alphas.Num() ? .01f : Alphas[Idx];
+            }
+            else if (Activation == TEXT("ThresholdedRelu"))
+            { 
+                ActivationOpDesc.Type = DML_OPERATOR_ACTIVATION_THRESHOLDED_RELU;
+                ActivationOpInnerDesc.LeakyRelu.Alpha = Idx >= Alphas.Num() ? 1.0f : Alphas[Idx];
+            }
+            else if (Activation == TEXT("Tanh"))
+            { 
+                ActivationOpDesc.Type = DML_OPERATOR_ACTIVATION_TANH;
+            }
+            else if (Activation == TEXT("ScaledTanh"))
+            { 
+                ActivationOpDesc.Type = DML_OPERATOR_ACTIVATION_SCALED_TANH;
+                ActivationOpInnerDesc.ScaledTanh.Alpha = Idx >= Alphas.Num() ? 1.0f : Alphas[Idx];
+                ActivationOpInnerDesc.ScaledTanh.Beta = Idx >= Betas.Num() ? 1.0f : Betas[Idx];
+            }
+            else if (Activation == TEXT("Sigmoid"))
+            { 
+                ActivationOpDesc.Type = DML_OPERATOR_ACTIVATION_SIGMOID;
+            }
+            else if (Activation == TEXT("HardSigmoid"))
+            { 
+                ActivationOpDesc.Type = DML_OPERATOR_ACTIVATION_HARD_SIGMOID;
+                ActivationOpInnerDesc.HardSigmoid.Alpha = Idx >= Alphas.Num() ? 0.2f : Alphas[Idx];
+                ActivationOpInnerDesc.HardSigmoid.Beta = Idx >= Betas.Num() ? 0.5f : Betas[Idx];
+            }
+            else if (Activation == TEXT("Elu"))
+            { 
+                ActivationOpDesc.Type = DML_OPERATOR_ACTIVATION_ELU;
+                ActivationOpInnerDesc.Elu.Alpha = Idx >= Alphas.Num() ? 1.0f : Alphas[Idx];
+            }
+            else if (Activation == TEXT("Softsign"))
+            { 
+                ActivationOpDesc.Type = DML_OPERATOR_ACTIVATION_SOFTSIGN;
+            }
+            else if (Activation == TEXT("Softplus"))
+            { 
+                ActivationOpDesc.Type = DML_OPERATOR_ACTIVATION_SOFTPLUS;
+            }
+            else
+            {
+                UE_LOG(LogNNE, Error, TEXT("Failed to initialize GRU: unsupported activation: %s"), *Activation);
+				return false;
+            }
+		}
+
+		return true;
+	}
+
+public:
+
+	static FOperatorDml* Create()
+	{
+		return new FOperatorDmlGRU();
+	}
+
+	//
+	//
+	//
+	virtual bool Initialize(IDMLDevice* Device, TArrayView<const NNECore::Internal::FTensor> InputTensors, TArrayView<const NNECore::Internal::FTensor> OutputTensors, const NNECore::FAttributeMap& Attributes) override
+	{
+		
+#define SET_DMLTENSORDESC_FROM_TENSOR_COND(OpDesc, Tensor, Name, Cond)\
+		DmlUtil::FTensorDesc	Dml##Name{};\
+		if(Cond)\
+		{\
+			NNECore::Internal::FTensor Name = To4DTensor(Tensor);\
+			if (!InitDmlTensorDesc(Dml##Name, Name))\
+			{\
+				UE_LOG(LogNNE, Error, TEXT("Failed to initialize GRU's "#Name" for DML inference"));\
+				return false;\
+			}\
+			##OpDesc.##Name = &Dml##Name.Desc;\
+		}\
+		else\
+		{\
+			##OpDesc.##Name = nullptr;\
+		}
+
+#define SET_DMLTENSORDESC_FROM_TENSOR(OpDesc, Tensor, Name)	SET_DMLTENSORDESC_FROM_TENSOR_COND(OpDesc, Tensor, Name, true)
+
+
+		check(InputTensors.Num() >= 3 && InputTensors.Num() <= 6);
+		check(OutputTensors.Num() >= 0 && OutputTensors.Num() <= 2);
+
+		check(InputTensors[0].GetShape().Rank() == 3);
+		check(InputTensors[1].GetShape().Rank() == 3);
+		check(InputTensors[2].GetShape().Rank() == 3);
+		if(InputTensors.Num() >= 4)
+		{
+			check(InputTensors[3].GetShape().Rank() == 2);
+		}
+		if(InputTensors.Num() >= 5)
+		{
+			check(InputTensors[4].GetShape().Rank() == 1);
+		}
+		if(InputTensors.Num() >= 6)
+		{
+			check(InputTensors[5].GetShape().Rank() == 3);
+		}
+		if(OutputTensors.Num() >= 1)
+		{
+			check(OutputTensors[0].GetShape().Rank() == 4);
+		}
+		if(OutputTensors.Num() >= 2)
+		{
+			check(OutputTensors[1].GetShape().Rank() == 3);
+		}
+
+		const int SeqLength = InputTensors[0].GetShape().GetData()[0];
+		const int BatchSize = InputTensors[0].GetShape().GetData()[1];
+		const int InputSize = InputTensors[0].GetShape().GetData()[2];
+		const int NumDirections = InputTensors[1].GetShape().GetData()[0];
+		const int HiddenSize = InputTensors[2].GetShape().GetData()[2];
+
+		
+		check(HiddenSize * 3 == InputTensors[1].GetShape().GetData()[1]);
+		check(InputSize == InputTensors[1].GetShape().GetData()[2]);
+		check(NumDirections == InputTensors[2].GetShape().GetData()[0]);
+		check(HiddenSize * 3 == InputTensors[2].GetShape().GetData()[1]);
+		if(InputTensors.Num() >= 4)
+		{
+			check(NumDirections == InputTensors[3].GetShape().GetData()[0]);
+			check(HiddenSize * 6 == InputTensors[3].GetShape().GetData()[1]);
+		}
+		if(InputTensors.Num() >= 5)
+		{
+			check(BatchSize == InputTensors[4].GetShape().GetData()[0]);
+		}
+		if(InputTensors.Num() >= 6)
+		{
+			check(NumDirections == InputTensors[5].GetShape().GetData()[0]);
+			check(BatchSize == InputTensors[5].GetShape().GetData()[1]);
+			check(HiddenSize == InputTensors[5].GetShape().GetData()[2]);
+		}
+		if(OutputTensors.Num() >= 1)
+		{
+			check(SeqLength == OutputTensors[0].GetShape().GetData()[0]);
+			check(NumDirections == OutputTensors[0].GetShape().GetData()[1]);
+			check(BatchSize == OutputTensors[0].GetShape().GetData()[2]);
+			check(HiddenSize == OutputTensors[0].GetShape().GetData()[3]);
+		}
+		if(OutputTensors.Num() >= 2)
+		{
+			check(NumDirections == OutputTensors[1].GetShape().GetData()[0]);
+			check(BatchSize == OutputTensors[1].GetShape().GetData()[1]);
+			check(HiddenSize == OutputTensors[1].GetShape().GetData()[2]);
+		}
+
+		DML_GRU_OPERATOR_DESC	DmlGRUOpDesc{};
+
+		SET_DMLTENSORDESC_FROM_TENSOR(DmlGRUOpDesc, InputTensors[0], InputTensor)
+		SET_DMLTENSORDESC_FROM_TENSOR(DmlGRUOpDesc, InputTensors[1], WeightTensor)
+		SET_DMLTENSORDESC_FROM_TENSOR(DmlGRUOpDesc, InputTensors[2], RecurrenceTensor)
+		SET_DMLTENSORDESC_FROM_TENSOR_COND(DmlGRUOpDesc, InputTensors[3], BiasTensor, InputTensors.Num() >= 4)
+		SET_DMLTENSORDESC_FROM_TENSOR_COND(DmlGRUOpDesc, InputTensors[4], SequenceLengthsTensor, InputTensors.Num() >= 5)
+		SET_DMLTENSORDESC_FROM_TENSOR_COND(DmlGRUOpDesc, InputTensors[5], HiddenInitTensor, InputTensors.Num() >= 6)
+		SET_DMLTENSORDESC_FROM_TENSOR_COND(DmlGRUOpDesc, OutputTensors[0], OutputSequenceTensor, OutputTensors.Num() >= 1)
+		SET_DMLTENSORDESC_FROM_TENSOR_COND(DmlGRUOpDesc, OutputTensors[1], OutputSingleTensor, OutputTensors.Num() >= 2)
+
+
+		DmlGRUOpDesc.Direction = DirectionFromString(Attributes.GetValueOrDefault<FString>(TEXT("direction"), TEXT("forward")));
+		
+		const TArray<FString> DefaultActivations = {"Sigmoid", "Tanh"};
+		const TArray<FString> DefaultActivationsBidirectional = {"Sigmoid", "Tanh", "Sigmoid", "Tanh"};
+		TArray<FString> Activations = 
+			Attributes.GetValueOrDefault<TArray<FString>>(
+				TEXT("activations"), 
+				DmlGRUOpDesc.Direction == DML_RECURRENT_NETWORK_DIRECTION_BIDIRECTIONAL ?
+					DefaultActivationsBidirectional : DefaultActivations);
+		
+		check(DmlGRUOpDesc.Direction == DML_RECURRENT_NETWORK_DIRECTION_BIDIRECTIONAL ?
+				(Activations.Num() == 4) : (Activations.Num() == 2));
+
+		
+		if(Attributes.GetAttributeValue(TEXT("clip")) != nullptr)
+		{
+			UE_LOG(LogNNE, Error, TEXT("GRU's clip attribute not supported for DML inference"));
+			return false;
+		}
+		
+		
+		const int HiddenSizeParam = Attributes.GetValueOrDefault<int>(TEXT("hidden_size"), HiddenSize);
+		check(HiddenSize == HiddenSizeParam);
+
+		const int Layout = Attributes.GetValueOrDefault<int>(TEXT("layout"), 0);
+		if(Layout != 0)
+		{
+			UE_LOG(LogNNE, Error, TEXT("GRU's layout attribute should be 0 for DML inference"));
+			return false;
+		}
+
+		//TODO: a way should be found to share implementation of activations from NNEDmlOperatorElementwise.cpp
+
+		TArray<float> Alphas = 
+			Attributes.GetValueOrDefault<TArray<float>>(
+				TEXT("activation_alpha"), 
+				TArray<float>());
+
+		TArray<float> Betas = 
+			Attributes.GetValueOrDefault<TArray<float>>(
+				TEXT("activation_beta"), 
+				TArray<float>());
+
+		TArray<DML_OPERATOR_DESC> ActivationOpDescs;
+		TArray<FDMLGRUActivationOpDescUnion> ActivationOpInnerDescs;
+		
+		if(!InitActivations(Activations, Alphas, 
+						 Betas, ActivationOpDescs,
+						 ActivationOpInnerDescs))
+		{
+			return false;
+		}
+
+		DmlGRUOpDesc.ActivationDescCount = (uint32_t) ActivationOpDescs.Num();
+		DmlGRUOpDesc.ActivationDescs = ActivationOpDescs.GetData();
+		
+		DmlGRUOpDesc.LinearBeforeReset = 
+			Attributes.GetValueOrDefault<int>(TEXT("linear_before_reset"), 0) == 0 ? false : true;
+
+		DML_OPERATOR_DESC DmlOpDesc{};
+
+		DmlOpDesc.Type = DML_OPERATOR_GRU;
+		DmlOpDesc.Desc = &DmlGRUOpDesc;
+
+		return CreateOperator(Device, DmlOpDesc);
+#undef SET_DMLTENSORDESC_FROM_TENSOR
+#undef SET_DMLTENSORDESC_FROM_TENSOR_COND
+	}
+};
+
+// Register GRU operator on Module startup
+NNE_DML_REGISTER_OP(GRU)
+
+} // namespace UE::NNERuntimeRDG::Private::Dml
+
+#endif // NNE_USE_DIRECTML
