Add RDG DirectML InstanceNormalization operator
Refactored DmlUtils::FTensorDesc

#rnx

diff --git a/Engine/Plugins/Experimental/NNE/Source/NNEQA/Resources/NNEQATestDesc.json b/Engine/Plugins/Experimental/NNE/Source/NNEQA/Resources/NNEQATestDesc.json
--- a/Engine/Plugins/Experimental/NNE/Source/NNEQA/Resources/NNEQATestDesc.json
+++ b/Engine/Plugins/Experimental/NNE/Source/NNEQA/Resources/NNEQATestDesc.json
@@ -311,18 +311,12 @@
 		},
 		{
 			"category": "InstanceNormalization",
-			"runtimes": [
-				{ "name": "NNERuntimeRDGDml", "skip": true }
-			],
 			"targets": [
 				{ "target": "InstanceNormalization", "absolute_tolerance": 1e-5 }
 			]
 		},
 		{
 			"category": "Upsample",
-			"runtimes": [
-				{ "name": "NNERuntimeRDGDml", "skip": true	}
-			],
 			"targets": [
 				{ "target": "Upsample" }
 			]
@@ -658,15 +652,38 @@
 		{
 			"name": "InstanceNormalization",
 			"datasets": [
-				{ "inputs": [ { "shape": [ 1, 1, 3 ] }, { "shape": [ 1 ] }, { "shape": [ 1 ] } ], "outputs": [ { "shape": [ 1, 1, 3 ] } ], "runtimes": [ { "name": "NNERuntimeORTDml", "skip": true, "comment": "Not supported" } ] },
-				{ "inputs": [ { "shape": [ 1, 2, 3 ] }, { "shape": [ 2 ] }, { "shape": [ 2 ] } ], "outputs": [ { "shape": [ 1, 2, 3 ] } ], "runtimes": [ { "name": "NNERuntimeORTDml", "skip": true, "comment": "Not supported" } ] },
-				{ "inputs": [ { "shape": [ 2, 2, 3 ] }, { "shape": [ 2 ] }, { "shape": [ 2 ] } ], "outputs": [ { "shape": [ 2, 2, 3 ] } ], "runtimes": [ { "name": "NNERuntimeORTDml", "skip": true, "comment": "Not supported" } ]},
-				{ "inputs": [ { "shape": [ 1, 2, 1, 3 ] }, { "shape": [ 2 ] }, { "shape": [ 2 ] } ], "outputs": [ { "shape": [ 1, 2, 1, 3 ] } ] },
-				{ "inputs": [ { "shape": [ 2, 2, 1, 3 ] }, { "shape": [ 2 ] }, { "shape": [ 2 ] } ], "outputs": [ { "shape": [ 2, 2, 1, 3 ] } ] },
-				{ "inputs": [ { "shape": [ 2, 2, 2, 2 ] }, { "shape": [ 2 ] }, { "shape": [ 2 ] } ], "outputs": [ { "shape": [ 2, 2, 2, 2 ] } ] },
+				{ 
+					"inputs": [ { "shape": [ 1, 1, 3 ] }, { "shape": [ 1 ] }, { "shape": [ 1 ] } ], "outputs": [ { "shape": [ 1, 1, 3 ] } ], 
+					"runtimes": [ 
+						{ "name": "NNERuntimeORTDml", "skip": true, "comment": "Not supported" }  
+					] 
+				},
+				{ 
+					"inputs": [ { "shape": [ 1, 2, 3 ] }, { "shape": [ 2 ] }, { "shape": [ 2 ] } ], "outputs": [ { "shape": [ 1, 2, 3 ] } ], 
+					"runtimes": [ 
+						{ "name": "NNERuntimeRDGDml", "skip": true, "comment": "Not supported" },
+						{ "name": "NNERuntimeORTDml", "skip": true, "comment": "Not supported" }
+					] 
+				},
+				{ 
+					"inputs": [ { "shape": [ 2, 2, 3 ] }, { "shape": [ 2 ] }, { "shape": [ 2 ] } ], "outputs": [ { "shape": [ 2, 2, 3 ] } ], 
+					"runtimes": [
+						{ "name": "NNERuntimeORTDml", "skip": true, "comment": "Not supported" } 
+					]
+				},
+				{ 
+					"inputs": [ { "shape": [ 1, 2, 1, 3 ] }, { "shape": [ 2 ] }, { "shape": [ 2 ] } ], "outputs": [ { "shape": [ 1, 2, 1, 3 ] } ] 
+				},
+				{ 
+					"inputs": [ { "shape": [ 2, 2, 1, 3 ] }, { "shape": [ 2 ] }, { "shape": [ 2 ] } ], "outputs": [ { "shape": [ 2, 2, 1, 3 ] } ]
+				},
+				{ 
+					"inputs": [ { "shape": [ 2, 2, 2, 2 ] }, { "shape": [ 2 ] }, { "shape": [ 2 ] } ], "outputs": [ { "shape": [ 2, 2, 2, 2 ] } ]
+				},
 				{
 					"inputs": [ { "shape": [ 2, 2, 4, 1, 3 ] }, { "shape": [ 2 ] }, { "shape": [ 2 ] } ], "outputs": [ { "shape": [ 2, 2, 4, 1, 3 ] } ],
 					"runtimes": [
+						{ "name": "NNERuntimeRDGDml", "skip": true, "comment": "Not supported" },
 						{ "name": "NNERuntimeRDGHlsl" },
 						{ "name": "NNERuntimeORTCpu" },
 						{ "name": "NNERuntimeORTDml", "skip": true, "comment": "Not supported" }
@@ -675,6 +692,7 @@
 				{
 					"inputs": [ { "shape": [ 2, 2, 4, 4 ] }, { "shape": [ 2 ] }, { "shape": [ 2 ] } ], "outputs": [ { "shape": [ 2, 2, 4, 4 ] } ],
 					"runtimes": [
+						{ "name": "NNERuntimeRDGDml" },
 						{ "name": "NNERuntimeRDGHlsl" },
 						{ "name": "NNERuntimeORTCpu" },
 						{ "name": "NNERuntimeORTDml" }
@@ -683,6 +701,7 @@
 				{
 					"inputs": [ { "shape": [ 2, 2, 8, 8 ] }, { "shape": [ 2 ] }, { "shape": [ 2 ] } ], "outputs": [ { "shape": [ 2, 2, 8, 8 ] } ],
 					"runtimes": [
+						{ "name": "NNERuntimeRDGDml" },
 						{ "name": "NNERuntimeRDGHlsl" },
 						{ "name": "NNERuntimeORTCpu" },
 						{ "name": "NNERuntimeORTDml" }
@@ -691,6 +710,7 @@
 				{
 					"inputs": [ { "shape": [ 2, 128, 8, 8 ] }, { "shape": [ 128 ] }, { "shape": [ 128 ] } ], "outputs": [ { "shape": [ 2, 128, 8, 8 ] } ],
 					"runtimes": [
+						{ "name": "NNERuntimeRDGDml" },
 						{ "name": "NNERuntimeRDGHlsl" },
 						{ "name": "NNERuntimeORTCpu" },
 						{ "name": "NNERuntimeORTDml" }
@@ -699,6 +719,7 @@
 				{
 					"inputs": [ { "shape": [ 128, 2, 8, 8 ] }, { "shape": [ 2 ] }, { "shape": [ 2 ] } ], "outputs": [ { "shape": [ 128, 2, 8, 8 ] } ],
 					"runtimes": [
+						{ "name": "NNERuntimeRDGDml" },
 						{ "name": "NNERuntimeRDGHlsl" },
 						{ "name": "NNERuntimeORTCpu" },
 						{ "name": "NNERuntimeORTDml" }
@@ -707,6 +728,7 @@
 				{
 					"inputs": [ { "shape": [ 64, 64, 8, 8 ] }, { "shape": [ 64 ] }, { "shape": [ 64 ] } ], "outputs": [ { "shape": [ 64, 64, 8, 8 ] } ],
 					"runtimes": [
+						{ "name": "NNERuntimeRDGDml" },
 						{ "name": "NNERuntimeRDGHlsl" },
 						{ "name": "NNERuntimeORTCpu" },
 						{ "name": "NNERuntimeORTDml" }
@@ -715,6 +737,7 @@
 				{
 					"inputs": [ { "shape": [ 2, 2, 45, 45 ] }, { "shape": [ 2 ] }, { "shape": [ 2 ] } ], "outputs": [ { "shape": [ 2, 2, 45, 45 ] } ],
 					"runtimes": [
+						{ "name": "NNERuntimeRDGDml" },
 						{ "name": "NNERuntimeRDGHlsl" },
 						{ "name": "NNERuntimeORTCpu" },
 						{ "name": "NNERuntimeORTDml" }
@@ -723,6 +746,7 @@
 				{
 					"inputs": [ { "shape": [ 1, 128, 56, 56 ] }, { "shape": [ 128 ] }, { "shape": [ 128 ] } ], "outputs": [ { "shape": [ 1, 128, 56, 56 ] } ],
 					"runtimes": [
+						{ "name": "NNERuntimeRDGDml" },
 						{ "name": "NNERuntimeRDGHlsl" },
 						{ "name": "NNERuntimeORTCpu" },
 						{ "name": "NNERuntimeORTDml" }
diff --git a/Engine/Plugins/Experimental/NNE/Source/NNERuntimeRDG/Private/Dml/NNEDmlOperator.cpp b/Engine/Plugins/Experimental/NNE/Source/NNERuntimeRDG/Private/Dml/NNEDmlOperator.cpp
--- a/Engine/Plugins/Experimental/NNE/Source/NNERuntimeRDG/Private/Dml/NNEDmlOperator.cpp
+++ b/Engine/Plugins/Experimental/NNE/Source/NNERuntimeRDG/Private/Dml/NNEDmlOperator.cpp
@@ -10,6 +10,176 @@ namespace UE::NNERuntimeRDG::Private::Dml
 
 namespace DmlUtil
 {
+	bool FTensorDesc::InitFromTensor(const NNECore::Internal::FTensor& Tensor, int32 MinTensorRank, TConstArrayView<uint32> BroadcastShape)
+	{
+		Reset();
+
+		DML_TENSOR_DATA_TYPE DmlDataType = DmlUtil::GetTensorDataType(Tensor.GetDataType());
+
+		if (DmlDataType == DML_TENSOR_DATA_TYPE_UNKNOWN)
+		{
+			return false;
+		}
+
+		TConstArrayView<uint32> InShape = Tensor.GetShape().GetData();
+
+		ElemSizeInBytes = Tensor.GetElemByteSize();
+
+		if (BroadcastShape.IsEmpty())
+		{
+			SetShape(InShape, MinTensorRank);
+		}
+		else
+		{
+			SetShapeAndStrides(InShape, BroadcastShape);
+		}
+
+		//Note: We should support tensor padding using strides defined in FTensorDesc
+		//DmlUtil::SetTensorStrides(DmlTensorDesc, TensorDesc.Strides);
+
+		Update(DmlDataType, Tensor.HasPreparedData());
+
+		return true;
+	}
+
+	bool FTensorDesc::InitFromTensor1D(const NNECore::Internal::FTensor& Tensor, int32 Rank)
+	{
+		Reset();
+
+		if (Tensor.GetShape().Rank() != 1)
+		{
+			return false;
+		}
+
+		DML_TENSOR_DATA_TYPE DmlDataType = DmlUtil::GetTensorDataType(Tensor.GetDataType());
+
+		if (DmlDataType == DML_TENSOR_DATA_TYPE_UNKNOWN)
+		{
+			return false;
+		}
+		
+		ElemSizeInBytes = Tensor.GetElemByteSize();
+
+		SetShape1D(Tensor.GetShape().GetData()[0], Rank);
+		Update(DmlDataType);
+
+		return true;
+	}
+
+	void FTensorDesc::Reset()
+	{
+		BuffDesc = DML_BUFFER_TENSOR_DESC{};
+		Desc = DML_TENSOR_DESC{};
+		Sizes.Reset();
+		Strides.Reset();
+		ElemSizeInBytes = 0;
+	}
+
+	void FTensorDesc::SetShape(TConstArrayView<uint32> Shape, int32 MinTensorRank)
+	{
+		Sizes = FSmallUIntArray(Shape.GetData(), Shape.Num());
+		const int32 DimOffset = MinTensorRank - Sizes.Num();
+
+		for (int Dim = 0; Dim < DimOffset; ++Dim)
+		{
+			Sizes.Insert(1, 0);
+		}
+	}
+
+	void FTensorDesc::SetShapeAndStrides(TConstArrayView<uint32> Shape, TConstArrayView<uint32> BroadcastShape)
+	{
+		const uint32 TargetDimension = BroadcastShape.Num() != -1 ? BroadcastShape.Num() : Shape.Num();
+		checkf(BroadcastShape.Num() >= Shape.Num(), TEXT("Can't broadcast tensor from rank %d to rank %d, should be inferior or equal."), Shape.Num(), TargetDimension);
+		
+		Sizes.SetNum(TargetDimension);
+		Strides.SetNum(TargetDimension);
+
+		const int32 DimensionOffset = int32(TargetDimension - Shape.Num());
+		
+		for (int32 i = 0; i < (int32) TargetDimension; ++i)
+		{
+			Sizes[i] = i < DimensionOffset ? 1 : Shape[i - DimensionOffset];
+		}
+
+		uint32 CurrStride = 1;
+
+		for (int32 i = TargetDimension - 1; i >= 0; --i)
+		{
+			const bool bBroadcast = Sizes[i] < BroadcastShape[i];
+
+			Strides[i] = bBroadcast ? 0 : CurrStride;
+			CurrStride *= Sizes[i];
+
+			Sizes[i] = BroadcastShape[i];
+		}
+	}
+
+	void FTensorDesc::SetShape1D(uint32 Dimension, int32 Rank)
+	{
+		Sizes.Reset();
+
+		Sizes.Add(1);
+		Sizes.Add(Dimension);
+		Sizes.Add(1);
+
+		for (int32 Dim = 3; Dim < Rank; ++Dim)
+		{
+			Sizes.Add(1);
+		}
+	}
+
+	void FTensorDesc::Update(DML_TENSOR_DATA_TYPE DataType, bool bHasWeightData)
+	{
+		BuffDesc.DataType = DataType;
+		BuffDesc.Flags = bHasWeightData ? DML_TENSOR_FLAG_OWNED_BY_DML : DML_TENSOR_FLAG_NONE;
+		BuffDesc.DimensionCount = Sizes.Num();
+		BuffDesc.Sizes = Sizes.GetData();
+		BuffDesc.Strides = Strides.IsEmpty() ? nullptr : Strides.GetData();
+		BuffDesc.TotalTensorSizeInBytes = CalculateBufferSize();
+
+		Desc = DML_TENSOR_DESC{ DML_TENSOR_TYPE_BUFFER, &BuffDesc };
+	}
+
+	uint64 FTensorDesc::CalculateBufferSize()
+	{
+		if (ElemSizeInBytes == 0)
+		{
+			return 0;
+		}
+
+		uint64 MinSizeInBytes = 0;
+		
+		if (Strides.IsEmpty())
+		{
+			MinSizeInBytes = Sizes[0];
+
+			for (int32 i = 1; i < Sizes.Num(); ++i)
+			{
+				MinSizeInBytes *= Sizes[i];
+			}
+
+			MinSizeInBytes *= ElemSizeInBytes;
+		}
+		else
+		{
+			uint32 IndexOfLastElement = 0;
+
+			for (int32 i = 0; i < Sizes.Num(); ++i)
+			{
+				IndexOfLastElement += (Sizes[i] - 1) * Strides[i];
+			}
+
+			MinSizeInBytes = (static_cast<uint64>(IndexOfLastElement) + 1) * ElemSizeInBytes;
+		}
+
+		// Round up to the nearest 4 bytes
+		MinSizeInBytes = (MinSizeInBytes + 3) & ~3ull;
+
+		return MinSizeInBytes;
+	}
+
+
+
 	void SetTensorStrides(FTensorDesc& TensorDesc, const NNECore::Internal::FTensor& InputDesc)
 	{
 		uint32 CurrStride = 1;
@@ -147,9 +317,9 @@ IDMLOperator* FOperatorDml::GetOperator()
 	return DmlOp;
 }
 
-bool FOperatorDml::InitDmlTensorDesc(DmlUtil::FTensorDesc& DmlTensorDesc, const NNECore::Internal::FTensor& TensorDesc)
+bool FOperatorDml::InitDmlTensorDesc(DmlUtil::FTensorDesc& DmlTensorDesc, const NNECore::Internal::FTensor& Tensor)
 {
-	DML_TENSOR_DATA_TYPE DmlDataType = DmlUtil::GetTensorDataType(TensorDesc.GetDataType());
+	DML_TENSOR_DATA_TYPE DmlDataType = DmlUtil::GetTensorDataType(Tensor.GetDataType());
 
 	if (DmlDataType == DML_TENSOR_DATA_TYPE_UNKNOWN)
 	{
@@ -159,7 +329,7 @@ bool FOperatorDml::InitDmlTensorDesc(DmlUtil::FTensorDesc& DmlTensorDesc, const
 		return false;
 	}
 
-	DmlTensorDesc.Sizes = TensorDesc.GetShape().GetData();
+	DmlTensorDesc.Sizes = Tensor.GetShape().GetData();
 	//Note: We should support tensor padding using strides defined in FTensorDesc
 	//DmlUtil::SetTensorStrides(DmlTensorDesc, TensorDesc.Strides);
 		
@@ -167,20 +337,20 @@ bool FOperatorDml::InitDmlTensorDesc(DmlUtil::FTensorDesc& DmlTensorDesc, const
 
 	BuffDesc = DML_BUFFER_TENSOR_DESC{};
 	BuffDesc.DataType = DmlDataType;
-	BuffDesc.Flags = TensorDesc.HasPreparedData() ? DML_TENSOR_FLAG_OWNED_BY_DML : DML_TENSOR_FLAG_NONE;
-	BuffDesc.DimensionCount = TensorDesc.GetShape().Rank();
+	BuffDesc.Flags = Tensor.HasPreparedData() ? DML_TENSOR_FLAG_OWNED_BY_DML : DML_TENSOR_FLAG_NONE;
+	BuffDesc.DimensionCount = Tensor.GetShape().Rank();
 	BuffDesc.Sizes = DmlTensorDesc.Sizes.GetData();
 	BuffDesc.Strides = nullptr;
-	BuffDesc.TotalTensorSizeInBytes = TensorDesc.GetDataSize();
+	BuffDesc.TotalTensorSizeInBytes = Tensor.GetDataSize(); // DmlUtil::CalculateBufferSize(DmlTensorDesc, Tensor);
 
 	DmlTensorDesc.Desc = DML_TENSOR_DESC{ DML_TENSOR_TYPE_BUFFER, &DmlTensorDesc.BuffDesc };
 
 	return true;
 }
 
-bool FOperatorDml::InitDmlTensorDesc(DmlUtil::FTensorDesc& DmlTensorDesc, const NNECore::Internal::FTensor& TensorDesc, const NNECore::Internal::FTensor& BroadcastDesc)
+bool FOperatorDml::InitDmlTensorDesc(DmlUtil::FTensorDesc& DmlTensorDesc, const NNECore::Internal::FTensor& Tensor, const NNECore::Internal::FTensor& Broadcast)
 {
-	DML_TENSOR_DATA_TYPE DmlDataType = DmlUtil::GetTensorDataType(TensorDesc.GetDataType());
+	DML_TENSOR_DATA_TYPE DmlDataType = DmlUtil::GetTensorDataType(Tensor.GetDataType());
 
 	if (DmlDataType == DML_TENSOR_DATA_TYPE_UNKNOWN)
 	{
@@ -190,27 +360,27 @@ bool FOperatorDml::InitDmlTensorDesc(DmlUtil::FTensorDesc& DmlTensorDesc, const
 		return false;
 	}
 
-	if (DmlUtil::IsSameShape(TensorDesc, BroadcastDesc))
+	if (DmlUtil::IsSameShape(Tensor, Broadcast))
 	{
-		DmlTensorDesc.Sizes = TensorDesc.GetShape().GetData();
-		DmlUtil::SetTensorStrides(DmlTensorDesc, TensorDesc);
+		DmlTensorDesc.Sizes = Tensor.GetShape().GetData();
+		DmlUtil::SetTensorStrides(DmlTensorDesc, Tensor);
 	}
-	else if (TensorDesc.GetShape().Rank() > BroadcastDesc.GetShape().Rank())
+	else if (Tensor.GetShape().Rank() > Broadcast.GetShape().Rank())
 	{
 		return false;
 	}
 	else
 	{
-		DmlUtil::SetTensorSizesAndStridesForBroadcast(DmlTensorDesc, TensorDesc, BroadcastDesc);
+		DmlUtil::SetTensorSizesAndStridesForBroadcast(DmlTensorDesc, Tensor, Broadcast);
 	}
 
 	//UE_LOG(LogNNE, Warning, TEXT("DmlTensorDesc:%d,%d,%d -> %d,%d,%d"),
-	//	TensorDesc.Sizes[0],
-	//	TensorDesc.Dimension > 1 ? TensorDesc.Sizes[1] : 0,
-	//	TensorDesc.Dimension > 2 ? TensorDesc.Sizes[2] : 0,
-	//	DmlTensorDesc.Sizes[0],
-	//	DmlTensorDesc.Dimension > 1 ? DmlTensorDesc.Sizes[1] : 0,
-	//	DmlTensorDesc.Dimension > 2 ? DmlTensorDesc.Sizes[2] : 0
+	//	Tensor.Sizes[0],
+	//	Tensor.Dimension > 1 ? Tensor.Sizes[1] : 0,
+	//	Tensor.Dimension > 2 ? Tensor.Sizes[2] : 0,
+	//	DmlTensor.Sizes[0],
+	//	DmlTensor.Dimension > 1 ? DmlTensor.Sizes[1] : 0,
+	//	DmlTensor.Dimension > 2 ? DmlTensor.Sizes[2] : 0
 	//);
 
 	//UE_LOG(LogNNE, Warning, TEXT("DmlTensorStrides:%d,%d,%d"),
@@ -224,11 +394,11 @@ bool FOperatorDml::InitDmlTensorDesc(DmlUtil::FTensorDesc& DmlTensorDesc, const
 	BuffDesc = DML_BUFFER_TENSOR_DESC{};
 
 	BuffDesc.DataType = DmlDataType;
-	BuffDesc.Flags = TensorDesc.HasPreparedData() ? DML_TENSOR_FLAG_OWNED_BY_DML : DML_TENSOR_FLAG_NONE;
+	BuffDesc.Flags = Tensor.HasPreparedData() ? DML_TENSOR_FLAG_OWNED_BY_DML : DML_TENSOR_FLAG_NONE;
 	BuffDesc.DimensionCount = DmlTensorDesc.Sizes.Num();
 	BuffDesc.Sizes = DmlTensorDesc.Sizes.GetData();
 	BuffDesc.Strides = DmlTensorDesc.Strides.GetData();
-	BuffDesc.TotalTensorSizeInBytes = DmlUtil::CalculateBufferSize(DmlTensorDesc, TensorDesc);
+	BuffDesc.TotalTensorSizeInBytes = DmlUtil::CalculateBufferSize(DmlTensorDesc, Tensor);
 		
 	DmlTensorDesc.Desc = DML_TENSOR_DESC{ DML_TENSOR_TYPE_BUFFER, &DmlTensorDesc.BuffDesc };
 
diff --git a/Engine/Plugins/Experimental/NNE/Source/NNERuntimeRDG/Private/Dml/NNEDmlOperator.h b/Engine/Plugins/Experimental/NNE/Source/NNERuntimeRDG/Private/Dml/NNEDmlOperator.h
--- a/Engine/Plugins/Experimental/NNE/Source/NNERuntimeRDG/Private/Dml/NNEDmlOperator.h
+++ b/Engine/Plugins/Experimental/NNE/Source/NNERuntimeRDG/Private/Dml/NNEDmlOperator.h
@@ -25,6 +25,11 @@ static FDmlOperator##OpName##Registrator RegisterDmlOperator##OpName;
 
 namespace UE::NNERuntimeRDG::Private::Dml
 {
+
+static constexpr uint32_t NcdhwDimensionCount = 5;
+static constexpr uint32_t NcdhwSpatialDimensionCount = 3;
+static constexpr uint32_t NonspatialDimensionCount = 2; // The batch and channel dimensions of NCW, NCHW, NCDHW....
+
 //
 //
 //
@@ -40,12 +45,28 @@ public:
 
 namespace DmlUtil
 {
+	using FSmallUIntArray = TArray<uint32, TInlineAllocator<NNECore::FTensorShape::MaxRank>>;
+
 	struct FTensorDesc
 	{
-		DML_BUFFER_TENSOR_DESC												BuffDesc;
-		DML_TENSOR_DESC														Desc;
-		TArray<uint32, TInlineAllocator<NNECore::FTensorShape::MaxRank>>	Sizes;
-		TArray<uint32, TInlineAllocator<NNECore::FTensorShape::MaxRank>>	Strides;
+		DML_BUFFER_TENSOR_DESC	BuffDesc;
+		DML_TENSOR_DESC			Desc;
+		FSmallUIntArray			Sizes;
+		FSmallUIntArray			Strides;
+		uint64					ElemSizeInBytes;
+
+		bool InitFromTensor(const NNECore::Internal::FTensor& InputDesc, int32 MinTensorRank, TConstArrayView<uint32> Broadcast = MakeArrayView((uint32*) nullptr, 0));
+		bool InitFromTensor1D(const NNECore::Internal::FTensor& InputDesc, int32 Rank);
+
+	private:
+		void Reset();
+
+		void SetShape(TConstArrayView<uint32> Shape, int32 MinTensorRank);
+		void SetShapeAndStrides(TConstArrayView<uint32> Shape, TConstArrayView<uint32> BroadcastShape);
+		void SetShape1D(uint32 Dimension, int32 Rank);
+		
+		void Update(DML_TENSOR_DATA_TYPE DataType, bool bHasWeightData = false);
+		uint64 CalculateBufferSize();
 	};
 
 	extern void SetTensorStrides(FTensorDesc& TensorDesc, const NNECore::Internal::FTensor& InputDesc);
@@ -70,8 +91,8 @@ public:
 
 protected:
 
-	bool InitDmlTensorDesc(DmlUtil::FTensorDesc& DmlTensorDesc, const NNECore::Internal::FTensor& TensorDesc);
-	bool InitDmlTensorDesc(DmlUtil::FTensorDesc& DmlTensorDesc, const NNECore::Internal::FTensor& TensorDesc, const NNECore::Internal::FTensor& BroadcastDesc);
+	bool InitDmlTensorDesc(DmlUtil::FTensorDesc& DmlTensorDesc, const NNECore::Internal::FTensor& Tensor);
+	bool InitDmlTensorDesc(DmlUtil::FTensorDesc& DmlTensorDesc, const NNECore::Internal::FTensor& Tensor, const NNECore::Internal::FTensor& Broadcast);
 	
 	bool CreateOperator(IDMLDevice* Device, const DML_OPERATOR_DESC& DmlOpDesc);
 
diff --git a/Engine/Plugins/Experimental/NNE/Source/NNERuntimeRDG/Private/Dml/NNEDmlOperatorConvolution.cpp b/Engine/Plugins/Experimental/NNE/Source/NNERuntimeRDG/Private/Dml/NNEDmlOperatorConvolution.cpp
--- a/Engine/Plugins/Experimental/NNE/Source/NNERuntimeRDG/Private/Dml/NNEDmlOperatorConvolution.cpp
+++ b/Engine/Plugins/Experimental/NNE/Source/NNERuntimeRDG/Private/Dml/NNEDmlOperatorConvolution.cpp
@@ -12,13 +12,8 @@ namespace UE::NNERuntimeRDG::Private::Dml
 template <DML_CONVOLUTION_DIRECTION Direction>
 class FOperatorDmlConv : public FOperatorDml
 {
-	static constexpr uint32_t NcdhwDimensionCount = 5;
-	static constexpr uint32_t NcdhwSpatialDimensionCount = 3;
-	static constexpr uint32_t NonspatialDimensionCount = 2; // The batch and channel dimensions of NCW, NCHW, NCDHW....
-
 	using FSmallArray = TArray<uint32, TInlineAllocator<NcdhwSpatialDimensionCount>>;
 	using FIntArray = TArray<int32>;
-	using FUIntArray = TArray<uint32>;
 	
 	enum EAutoPad
 	{
diff --git a/Engine/Plugins/Experimental/NNE/Source/NNERuntimeRDG/Private/Dml/NNEDmlOperatorInstanceNormalization.cpp b/Engine/Plugins/Experimental/NNE/Source/NNERuntimeRDG/Private/Dml/NNEDmlOperatorInstanceNormalization.cpp
new file mode 100644
--- /dev/null
+++ b/Engine/Plugins/Experimental/NNE/Source/NNERuntimeRDG/Private/Dml/NNEDmlOperatorInstanceNormalization.cpp
@@ -0,0 +1,133 @@
+// Copyright Epic Games, Inc. All Rights Reserved.
+
+#ifdef NNE_USE_DIRECTML
+#include "NNEDmlOperator.h"
+
+namespace UE::NNERuntimeRDG::Private::Dml
+{
+
+	class FOperatorDmlInstanceNormalization : public FOperatorDml
+	{
+		static constexpr float DefaultEpsilon = 0.00001f;
+
+	public:
+
+		static FOperatorDml* Create()
+		{
+			return new FOperatorDmlInstanceNormalization();
+		}
+
+		//
+		//
+		//
+		virtual bool Initialize(IDMLDevice* Device, TArrayView<const NNECore::Internal::FTensor> InputTensors, TArrayView<const NNECore::Internal::FTensor> OutputTensors, const NNECore::FAttributeMap& Attributes) override
+		{
+			check(InputTensors.Num() >= 1 && InputTensors.Num() <= 3);
+			check(OutputTensors.Num() == 1);
+
+			const NNECore::Internal::FTensor& InputTensor = InputTensors[0];
+			const NNECore::Internal::FTensor& OutputTensor = OutputTensors[0];
+
+			if (InputTensor.GetShape().Rank() > 8)
+			{
+				UE_LOG(LogNNE, Warning, TEXT("InstanceNormalization:InputTensor rank should be between 1 and 8, got:%d"), InputTensor.GetShape().Rank());
+				return false;
+			}
+
+			// Read attributes
+			float	Epsilon = Attributes.GetValueOrDefault<float>(TEXT("epsilon"), DefaultEpsilon);
+
+			// Initialize tensor descriptors
+			DmlUtil::FTensorDesc	DmlInputTensor{};
+			DmlUtil::FTensorDesc	DmlScalingTensor{};
+			DmlUtil::FTensorDesc	DmlBiasTensor{};
+			DmlUtil::FTensorDesc	DmlOutputTensor{};
+			
+			// Make sure that input is 4D
+			if (!DmlInputTensor.InitFromTensor(InputTensor, 4))
+			{
+				UE_LOG(LogNNE, Warning, TEXT("Failed to initialize tensor(s) for DML inference"));
+				return false;
+			}
+
+			if (InputTensors.Num() > 1)
+			{
+				const NNECore::Internal::FTensor& ScaleTensor = InputTensors[1];
+
+				if (!DmlScalingTensor.InitFromTensor1D(ScaleTensor, DmlInputTensor.Sizes.Num()))
+				{
+					UE_LOG(LogNNE, Warning, TEXT("Failed to initialize tensor(s) for DML inference"));
+					return false;
+				}
+			}
+
+			if (InputTensors.Num() > 2)
+			{
+				const NNECore::Internal::FTensor& BiasTensor = InputTensors[2];
+
+				if (!DmlBiasTensor.InitFromTensor1D(BiasTensor, DmlInputTensor.Sizes.Num()))
+				{
+					UE_LOG(LogNNE, Warning, TEXT("Failed to initialize tensor(s) for DML inference"));
+					return false;
+				}
+			}
+
+			if (!DmlOutputTensor.InitFromTensor(OutputTensor, 4))
+			{
+				UE_LOG(LogNNE, Warning, TEXT("Failed to initialize tensor(s) for DML inference"));
+				return false;
+			}
+
+#ifdef NNE_USE_DML_OPERATOR_MEAN_VARIANCE_NORMALIZATION1
+
+			// "Instance" normalization is really spatial normalization,
+			// where the spatial channels are reduced and normalized, while
+			// batch and channel remain independent. So pass a list of axes
+			// just beyond the leading batch and channel dimensions (starting
+			// at axis 2 up to the last spatial dimension).
+			const int32 InputRank = DmlInputTensor.Sizes.Num();
+
+			DmlUtil::FSmallUIntArray Axes;
+			Axes.SetNumUninitialized(NcdhwSpatialDimensionCount);
+
+			for (int32 Dim = 0, Value = InputRank - Axes.Num(); Dim < Axes.Num(); ++Dim, ++Value)
+			{
+				Axes[Dim] = Value;
+			}
+
+			DML_MEAN_VARIANCE_NORMALIZATION1_OPERATOR_DESC	OpDesc{};
+
+			OpDesc.InputTensor = &DmlInputTensor.Desc;
+			OpDesc.ScaleTensor = InputTensors.Num() > 1 ? &DmlScalingTensor.Desc : nullptr;
+			OpDesc.BiasTensor = InputTensors.Num() > 2 ? &DmlBiasTensor.Desc : nullptr;
+			OpDesc.OutputTensor = &DmlOutputTensor.Desc;
+			OpDesc.AxisCount = Axes.Num();
+			OpDesc.Axes = Axes.GetData();
+			OpDesc.NormalizeVariance = true;
+			OpDesc.Epsilon = Epsilon;
+			
+			return CreateOperator(Device, DML_OPERATOR_DESC{ DML_OPERATOR_MEAN_VARIANCE_NORMALIZATION1, &OpDesc });
+#else
+
+			DML_MEAN_VARIANCE_NORMALIZATION_OPERATOR_DESC	OpDesc{};
+
+			OpDesc.InputTensor = &DmlInputTensor.Desc;
+			OpDesc.ScaleTensor = InputTensors.Num() > 1 ? &DmlScalingTensor.Desc : nullptr;
+			OpDesc.BiasTensor = InputTensors.Num() > 2 ? &DmlBiasTensor.Desc : nullptr;
+			OpDesc.OutputTensor = &DmlOutputTensor.Desc;
+			OpDesc.CrossChannel = false;
+			OpDesc.NormalizeVariance = true;
+			OpDesc.Epsilon = Epsilon;
+
+			return CreateOperator(Device, DML_OPERATOR_DESC{ DML_OPERATOR_MEAN_VARIANCE_NORMALIZATION, &OpDesc });
+#endif
+
+		}
+	};
+
+	// Register operator on Module startup
+	NNE_DML_REGISTER_OP(InstanceNormalization)
+
+} // namespace UE::NNERuntimeRDG::Private::Dml
+
+#endif // NNE_USE_DIRECTML
