NNE - RDG HLSL runtime - Add Concat, Slice and Unsqueeze to shape inference

Problem:
* Missing Concat, Slice and Unsqueeze operator in shapeinference/tensor constness propagation.

Solutions:
* Add them

Tests: all nne test green on editor win64 (but known convtranspose) + style transfer demo

diff --git a/Engine/Plugins/Experimental/NNE/Source/NNEQA/Private/NNEQAUnitTestHelper.cpp b/Engine/Plugins/Experimental/NNE/Source/NNEQA/Private/NNEQAUnitTestHelper.cpp
--- a/Engine/Plugins/Experimental/NNE/Source/NNEQA/Private/NNEQAUnitTestHelper.cpp
+++ b/Engine/Plugins/Experimental/NNE/Source/NNEQA/Private/NNEQAUnitTestHelper.cpp
@@ -2,8 +2,8 @@
 
 #include "NNEQAUnitTestHelper.h"
 #include "NNECoreTypes.h"
-#include "NNERuntimeRDGElementWiseBinaryHelper.h"
-#include "NNERuntimeRDGElementWiseUnaryHelper.h"
+#include "NNERuntimeRDGHelperElementWiseBinary.h"
+#include "NNERuntimeRDGHelperElementWiseUnary.h"
 
 #if WITH_DEV_AUTOMATION_TESTS
 #include "Misc/AutomationTest.h"
@@ -58,31 +58,31 @@ bool FShapeInferenceHelperUnitTestBase::TestBinaryOutputIsOnlyComputedWhenItShou
 	//Tests if output tensors are only computed if both inputs are constant and below a certain size
 	
 	FTensor Y = MakeTensor(TEXT("Y"), { 1 });
-	ElementWiseBinaryCPUHelper::Apply(OpType, XC1, XC1, Y);
+	CPUHelper::ElementWiseBinary::Apply(OpType, XC1, XC1, Y);
 	UTEST_TRUE(TEXT("Y const if both inputs are const"), Y.HasPreparedData());
 
 	Y = MakeTensor(TEXT("Y"), { 1 });
-	ElementWiseBinaryCPUHelper::Apply(OpType, X1, XC1, Y);
+	CPUHelper::ElementWiseBinary::Apply(OpType, X1, XC1, Y);
 	UTEST_FALSE(TEXT("Y not const when LHS not const"), Y.HasPreparedData());
 
 	Y = MakeTensor(TEXT("Y"), { 1 });
-	ElementWiseBinaryCPUHelper::Apply(OpType, XC1, X1, Y);
+	CPUHelper::ElementWiseBinary::Apply(OpType, XC1, X1, Y);
 	UTEST_FALSE(TEXT("Y not const when RHS not const"), Y.HasPreparedData());
 
 	Y = MakeTensor(TEXT("Y"), { 1 });
-	ElementWiseBinaryCPUHelper::Apply(OpType, X1, X1, Y);
+	CPUHelper::ElementWiseBinary::Apply(OpType, X1, X1, Y);
 	UTEST_FALSE(TEXT("Y not const when LHS and RHS not const"), Y.HasPreparedData());
 
 	Y = MakeTensor(TEXT("Y"), { 20 });
-	ElementWiseBinaryCPUHelper::Apply(OpType, XC20, XC1, Y);
+	CPUHelper::ElementWiseBinary::Apply(OpType, XC20, XC1, Y);
 	UTEST_FALSE(TEXT("Y not const when LHS too large"), Y.HasPreparedData());
 
 	Y = MakeTensor(TEXT("Y"), { 20 });
-	ElementWiseBinaryCPUHelper::Apply(OpType, XC1, XC20, Y);
+	CPUHelper::ElementWiseBinary::Apply(OpType, XC1, XC20, Y);
 	UTEST_FALSE(TEXT("Y not const when RHS too large"), Y.HasPreparedData());
 
 	Y = MakeTensor(TEXT("Y"), { 20 });
-	ElementWiseBinaryCPUHelper::Apply(OpType, XC20, XC20, Y);
+	CPUHelper::ElementWiseBinary::Apply(OpType, XC20, XC20, Y);
 	UTEST_FALSE(TEXT("Y not const when LHS and RHS too large"), Y.HasPreparedData());
 
 	return true;
@@ -97,15 +97,15 @@ bool FShapeInferenceHelperUnitTestBase::TestUnaryOutputIsOnlyComputedWhenItShoul
 	//Tests if output tensors are only computed if inputs are constant and below a certain size
 
 	FTensor Y = MakeTensor(TEXT("Y"), { 1 });
-	ElementWiseUnaryCPUHelper::Apply(OpType, XC1, 1.0f, 1.0f, 1.0f, Y);
+	CPUHelper::ElementWiseUnary::Apply(OpType, XC1, 1.0f, 1.0f, 1.0f, Y);
 	UTEST_TRUE(TEXT("Y const if input is const"), Y.HasPreparedData());
 
 	Y = MakeTensor(TEXT("Y"), { 1 });
-	ElementWiseUnaryCPUHelper::Apply(OpType, X1, 1.0f, 1.0f, 1.0f, Y);
+	CPUHelper::ElementWiseUnary::Apply(OpType, X1, 1.0f, 1.0f, 1.0f, Y);
 	UTEST_FALSE(TEXT("Y not const if input is not const"), Y.HasPreparedData());
 
 	Y = MakeTensor(TEXT("Y"), { 20 });
-	ElementWiseUnaryCPUHelper::Apply(OpType, XC20, 1.0f, 1.0f, 1.0f, Y);
+	CPUHelper::ElementWiseUnary::Apply(OpType, XC20, 1.0f, 1.0f, 1.0f, Y);
 	UTEST_FALSE(TEXT("Y not const if input is too large "), Y.HasPreparedData());
 
 	return true;
diff --git a/Engine/Plugins/Experimental/NNE/Source/NNEQA/Private/NNEQAUnitTestOpBinary.cpp b/Engine/Plugins/Experimental/NNE/Source/NNEQA/Private/NNEQAUnitTestOpBinary.cpp
--- a/Engine/Plugins/Experimental/NNE/Source/NNEQA/Private/NNEQAUnitTestOpBinary.cpp
+++ b/Engine/Plugins/Experimental/NNE/Source/NNEQA/Private/NNEQAUnitTestOpBinary.cpp
@@ -4,7 +4,7 @@
 #include "NNECoreTensor.h"
 #include "NNECoreTypes.h"
 #include "NNEQAUnitTestHelper.h"
-#include "NNERuntimeRDGElementWiseBinaryHelper.h"
+#include "NNERuntimeRDGHelperElementWiseBinary.h"
 
 #if WITH_DEV_AUTOMATION_TESTS
 #include "Misc/AutomationTest.h"
@@ -33,21 +33,21 @@ namespace UE::NNEQA::Private::NNERuntimeRDG::BinaryOp
 
 		//Test output tensor math is correct when both inputs are constant including broadcasting
 		FTensor Y = MakeTensor(TEXT("Y"), { 1 });
-		ElementWiseBinaryCPUHelper::Apply(EElementWiseBinaryOperatorType::Add, XC1, XC1, Y);
+		CPUHelper::ElementWiseBinary::Apply(EElementWiseBinaryOperatorType::Add, XC1, XC1, Y);
 		UTEST_EQUAL(TEXT("XC1+XC1"), Y.GetPreparedData<float>()[0], 1.0f + 1.0f);
 
 		Y = MakeTensor(TEXT("Y"), { 1,2 });
-		ElementWiseBinaryCPUHelper::Apply(EElementWiseBinaryOperatorType::Add, XC1x2, XC1, Y);
+		CPUHelper::ElementWiseBinary::Apply(EElementWiseBinaryOperatorType::Add, XC1x2, XC1, Y);
 		UTEST_EQUAL(TEXT("XC1x2+XC1[0]"), Y.GetPreparedData<float>()[0], 1.0f + 1.0f);
 		UTEST_EQUAL(TEXT("XC1x2+XC1[1]"), Y.GetPreparedData<float>()[1], 2.0f + 1.0f);
 
 		Y = MakeTensor(TEXT("Y"), { 1,2 });
-		ElementWiseBinaryCPUHelper::Apply(EElementWiseBinaryOperatorType::Add, XC1x2, XC1x2, Y);
+		CPUHelper::ElementWiseBinary::Apply(EElementWiseBinaryOperatorType::Add, XC1x2, XC1x2, Y);
 		UTEST_EQUAL(TEXT("XC1x2+XC1x2[0]"), Y.GetPreparedData<float>()[0], 1.0f + 1.0f);
 		UTEST_EQUAL(TEXT("XC1x2+XC1x2[1]"), Y.GetPreparedData<float>()[1], 2.0f + 2.0f);
 
 		Y = MakeTensor(TEXT("Y"), { 2,2 });
-		ElementWiseBinaryCPUHelper::Apply(EElementWiseBinaryOperatorType::Add, XC1x2, XC2x1, Y);
+		CPUHelper::ElementWiseBinary::Apply(EElementWiseBinaryOperatorType::Add, XC1x2, XC2x1, Y);
 		UTEST_EQUAL(TEXT("XC1x2+XC2x1[0]"), Y.GetPreparedData<float>()[0], 1.0f + 3.0f);
 		UTEST_EQUAL(TEXT("XC1x2+XC2x1[1]"), Y.GetPreparedData<float>()[1], 2.0f + 3.0f);
 		UTEST_EQUAL(TEXT("XC1x2+XC2x1[2]"), Y.GetPreparedData<float>()[2], 1.0f + 4.0f);
@@ -70,21 +70,21 @@ namespace UE::NNEQA::Private::NNERuntimeRDG::BinaryOp
 
 		//Test output tensor math is correct when both inputs are constant including broadcasting
 		FTensor Y = MakeTensor(TEXT("Y"), { 1 });
-		ElementWiseBinaryCPUHelper::Apply(EElementWiseBinaryOperatorType::Div, XC1, XC1, Y);
+		CPUHelper::ElementWiseBinary::Apply(EElementWiseBinaryOperatorType::Div, XC1, XC1, Y);
 		UTEST_EQUAL(TEXT("XC1/XC1"), Y.GetPreparedData<float>()[0], 1.0f / 1.0f);
 
 		Y = MakeTensor(TEXT("Y"), { 1,2 });
-		ElementWiseBinaryCPUHelper::Apply(EElementWiseBinaryOperatorType::Div, XC1x2, XC1, Y);
+		CPUHelper::ElementWiseBinary::Apply(EElementWiseBinaryOperatorType::Div, XC1x2, XC1, Y);
 		UTEST_EQUAL(TEXT("XC1x2/XC1[0]"), Y.GetPreparedData<float>()[0], 1.0f / 1.0f);
 		UTEST_EQUAL(TEXT("XC1x2/XC1[1]"), Y.GetPreparedData<float>()[1], 2.0f / 1.0f);
 
 		Y = MakeTensor(TEXT("Y"), { 1,2 });
-		ElementWiseBinaryCPUHelper::Apply(EElementWiseBinaryOperatorType::Div, XC1x2, XC1x2, Y);
+		CPUHelper::ElementWiseBinary::Apply(EElementWiseBinaryOperatorType::Div, XC1x2, XC1x2, Y);
 		UTEST_EQUAL(TEXT("XC1x2/XC1x2[0]"), Y.GetPreparedData<float>()[0], 1.0f / 1.0f);
 		UTEST_EQUAL(TEXT("XC1x2/XC1x2[1]"), Y.GetPreparedData<float>()[1], 2.0f / 2.0f);
 
 		Y = MakeTensor(TEXT("Y"), { 2,2 });
-		ElementWiseBinaryCPUHelper::Apply(EElementWiseBinaryOperatorType::Div, XC1x2, XC2x1, Y);
+		CPUHelper::ElementWiseBinary::Apply(EElementWiseBinaryOperatorType::Div, XC1x2, XC2x1, Y);
 		UTEST_EQUAL(TEXT("XC1x2/XC2x1[0]"), Y.GetPreparedData<float>()[0], 1.0f / 3.0f);
 		UTEST_EQUAL(TEXT("XC1x2/XC2x1[1]"), Y.GetPreparedData<float>()[1], 2.0f / 3.0f);
 		UTEST_EQUAL(TEXT("XC1x2/XC2x1[2]"), Y.GetPreparedData<float>()[2], 1.0f / 4.0f);
@@ -107,21 +107,21 @@ namespace UE::NNEQA::Private::NNERuntimeRDG::BinaryOp
 
 		//Test output tensor math is correct when both inputs are constant including broadcasting
 		FTensor Y = MakeTensor(TEXT("Y"), { 1 });
-		ElementWiseBinaryCPUHelper::Apply(EElementWiseBinaryOperatorType::Mul, XC1, XC1, Y);
+		CPUHelper::ElementWiseBinary::Apply(EElementWiseBinaryOperatorType::Mul, XC1, XC1, Y);
 		UTEST_EQUAL(TEXT("XC1*XC1"), Y.GetPreparedData<float>()[0], 1.0f * 1.0f);
 
 		Y = MakeTensor(TEXT("Y"), { 1,2 });
-		ElementWiseBinaryCPUHelper::Apply(EElementWiseBinaryOperatorType::Mul, XC1x2, XC1, Y);
+		CPUHelper::ElementWiseBinary::Apply(EElementWiseBinaryOperatorType::Mul, XC1x2, XC1, Y);
 		UTEST_EQUAL(TEXT("XC1x2*XC1[0]"), Y.GetPreparedData<float>()[0], 1.0f * 1.0f);
 		UTEST_EQUAL(TEXT("XC1x2*XC1[1]"), Y.GetPreparedData<float>()[1], 2.0f * 1.0f);
 
 		Y = MakeTensor(TEXT("Y"), { 1,2 });
-		ElementWiseBinaryCPUHelper::Apply(EElementWiseBinaryOperatorType::Mul, XC1x2, XC1x2, Y);
+		CPUHelper::ElementWiseBinary::Apply(EElementWiseBinaryOperatorType::Mul, XC1x2, XC1x2, Y);
 		UTEST_EQUAL(TEXT("XC1x2*XC1x2[0]"), Y.GetPreparedData<float>()[0], 1.0f * 1.0f);
 		UTEST_EQUAL(TEXT("XC1x2*XC1x2[1]"), Y.GetPreparedData<float>()[1], 2.0f * 2.0f);
 
 		Y = MakeTensor(TEXT("Y"), { 2,2 });
-		ElementWiseBinaryCPUHelper::Apply(EElementWiseBinaryOperatorType::Mul, XC1x2, XC2x1, Y);
+		CPUHelper::ElementWiseBinary::Apply(EElementWiseBinaryOperatorType::Mul, XC1x2, XC2x1, Y);
 		UTEST_EQUAL(TEXT("XC1x2*XC2x1[0]"), Y.GetPreparedData<float>()[0], 1.0f * 3.0f);
 		UTEST_EQUAL(TEXT("XC1x2*XC2x1[1]"), Y.GetPreparedData<float>()[1], 2.0f * 3.0f);
 		UTEST_EQUAL(TEXT("XC1x2*XC2x1[2]"), Y.GetPreparedData<float>()[2], 1.0f * 4.0f);
@@ -144,21 +144,21 @@ namespace UE::NNEQA::Private::NNERuntimeRDG::BinaryOp
 
 		//Test output tensor math is correct when both inputs are constant including broadcasting
 		FTensor Y = MakeTensor(TEXT("Y"), { 1 });
-		ElementWiseBinaryCPUHelper::Apply(EElementWiseBinaryOperatorType::Sub, XC1, XC1, Y);
+		CPUHelper::ElementWiseBinary::Apply(EElementWiseBinaryOperatorType::Sub, XC1, XC1, Y);
 		UTEST_EQUAL(TEXT("XC1-XC1"), Y.GetPreparedData<float>()[0], 1.0f - 1.0f);
 
 		Y = MakeTensor(TEXT("Y"), { 1,2 });
-		ElementWiseBinaryCPUHelper::Apply(EElementWiseBinaryOperatorType::Sub, XC1x2, XC1, Y);
+		CPUHelper::ElementWiseBinary::Apply(EElementWiseBinaryOperatorType::Sub, XC1x2, XC1, Y);
 		UTEST_EQUAL(TEXT("XC1x2-XC1[0]"), Y.GetPreparedData<float>()[0], 1.0f - 1.0f);
 		UTEST_EQUAL(TEXT("XC1x2-XC1[1]"), Y.GetPreparedData<float>()[1], 2.0f - 1.0f);
 
 		Y = MakeTensor(TEXT("Y"), { 1,2 });
-		ElementWiseBinaryCPUHelper::Apply(EElementWiseBinaryOperatorType::Sub, XC1x2, XC1x2, Y);
+		CPUHelper::ElementWiseBinary::Apply(EElementWiseBinaryOperatorType::Sub, XC1x2, XC1x2, Y);
 		UTEST_EQUAL(TEXT("XC1x2-XC1x2[0]"), Y.GetPreparedData<float>()[0], 1.0f - 1.0f);
 		UTEST_EQUAL(TEXT("XC1x2-XC1x2[1]"), Y.GetPreparedData<float>()[1], 2.0f - 2.0f);
 
 		Y = MakeTensor(TEXT("Y"), { 2,2 });
-		ElementWiseBinaryCPUHelper::Apply(EElementWiseBinaryOperatorType::Sub, XC1x2, XC2x1, Y);
+		CPUHelper::ElementWiseBinary::Apply(EElementWiseBinaryOperatorType::Sub, XC1x2, XC2x1, Y);
 		UTEST_EQUAL(TEXT("XC1x2-XC2x1[0]"), Y.GetPreparedData<float>()[0], 1.0f - 3.0f);
 		UTEST_EQUAL(TEXT("XC1x2-XC2x1[1]"), Y.GetPreparedData<float>()[1], 2.0f - 3.0f);
 		UTEST_EQUAL(TEXT("XC1x2-XC2x1[2]"), Y.GetPreparedData<float>()[2], 1.0f - 4.0f);
@@ -181,21 +181,21 @@ namespace UE::NNEQA::Private::NNERuntimeRDG::BinaryOp
 
 		//Test output tensor math is correct when both inputs are constant including broadcasting
 		FTensor Y = MakeTensor(TEXT("Y"), { 1 });
-		ElementWiseBinaryCPUHelper::Apply(EElementWiseBinaryOperatorType::Mod, XC1, XC1, Y);
+		CPUHelper::ElementWiseBinary::Apply(EElementWiseBinaryOperatorType::Mod, XC1, XC1, Y);
 		UTEST_EQUAL(TEXT("XC1%XC1"), Y.GetPreparedData<float>()[0], FMath::Fmod(1.0f, 1.0f));
 
 		Y = MakeTensor(TEXT("Y"), { 1,2 });
-		ElementWiseBinaryCPUHelper::Apply(EElementWiseBinaryOperatorType::Mod, XC1x2, XC1, Y);
+		CPUHelper::ElementWiseBinary::Apply(EElementWiseBinaryOperatorType::Mod, XC1x2, XC1, Y);
 		UTEST_EQUAL(TEXT("XC1x2%XC1[0]"), Y.GetPreparedData<float>()[0], FMath::Fmod(1.0f, 1.0f));
 		UTEST_EQUAL(TEXT("XC1x2%XC1[1]"), Y.GetPreparedData<float>()[1], FMath::Fmod(2.0f, 1.0f));
 
 		Y = MakeTensor(TEXT("Y"), { 1,2 });
-		ElementWiseBinaryCPUHelper::Apply(EElementWiseBinaryOperatorType::Mod, XC1x2, XC1x2, Y);
+		CPUHelper::ElementWiseBinary::Apply(EElementWiseBinaryOperatorType::Mod, XC1x2, XC1x2, Y);
 		UTEST_EQUAL(TEXT("XC1x2%XC1x2[0]"), Y.GetPreparedData<float>()[0], FMath::Fmod(1.0f, 1.0f));
 		UTEST_EQUAL(TEXT("XC1x2%XC1x2[1]"), Y.GetPreparedData<float>()[1], FMath::Fmod(2.0f, 2.0f));
 
 		Y = MakeTensor(TEXT("Y"), { 2,2 });
-		ElementWiseBinaryCPUHelper::Apply(EElementWiseBinaryOperatorType::Mod, XC1x2, XC2x1, Y);
+		CPUHelper::ElementWiseBinary::Apply(EElementWiseBinaryOperatorType::Mod, XC1x2, XC2x1, Y);
 		UTEST_EQUAL(TEXT("XC1x2%XC2x1[0]"), Y.GetPreparedData<float>()[0], FMath::Fmod(1.0f, 3.0f));
 		UTEST_EQUAL(TEXT("XC1x2%XC2x1[1]"), Y.GetPreparedData<float>()[1], FMath::Fmod(2.0f, 3.0f));
 		UTEST_EQUAL(TEXT("XC1x2%XC2x1[2]"), Y.GetPreparedData<float>()[2], FMath::Fmod(1.0f, 4.0f));
@@ -218,21 +218,21 @@ namespace UE::NNEQA::Private::NNERuntimeRDG::BinaryOp
 
 		//Test output tensor math is correct when both inputs are constant including broadcasting
 		FTensor Y = MakeTensor(TEXT("Y"), { 1 });
-		ElementWiseBinaryCPUHelper::Apply(EElementWiseBinaryOperatorType::Pow, XC1, XC1, Y);
+		CPUHelper::ElementWiseBinary::Apply(EElementWiseBinaryOperatorType::Pow, XC1, XC1, Y);
 		UTEST_EQUAL(TEXT("XC1^XC1"), Y.GetPreparedData<float>()[0], FMath::Pow(1.0f, 1.0f));
 
 		Y = MakeTensor(TEXT("Y"), { 1,2 });
-		ElementWiseBinaryCPUHelper::Apply(EElementWiseBinaryOperatorType::Pow, XC1x2, XC1, Y);
+		CPUHelper::ElementWiseBinary::Apply(EElementWiseBinaryOperatorType::Pow, XC1x2, XC1, Y);
 		UTEST_EQUAL(TEXT("XC1x2^XC1[0]"), Y.GetPreparedData<float>()[0], FMath::Pow(1.0f, 1.0f));
 		UTEST_EQUAL(TEXT("XC1x2^XC1[1]"), Y.GetPreparedData<float>()[1], FMath::Pow(2.0f, 1.0f));
 
 		Y = MakeTensor(TEXT("Y"), { 1,2 });
-		ElementWiseBinaryCPUHelper::Apply(EElementWiseBinaryOperatorType::Pow, XC1x2, XC1x2, Y);
+		CPUHelper::ElementWiseBinary::Apply(EElementWiseBinaryOperatorType::Pow, XC1x2, XC1x2, Y);
 		UTEST_EQUAL(TEXT("XC1x2^XC1x2[0]"), Y.GetPreparedData<float>()[0], FMath::Pow(1.0f, 1.0f));
 		UTEST_EQUAL(TEXT("XC1x2^XC1x2[1]"), Y.GetPreparedData<float>()[1], FMath::Pow(2.0f, 2.0f));
 
 		Y = MakeTensor(TEXT("Y"), { 2,2 });
-		ElementWiseBinaryCPUHelper::Apply(EElementWiseBinaryOperatorType::Pow, XC1x2, XC2x1, Y);
+		CPUHelper::ElementWiseBinary::Apply(EElementWiseBinaryOperatorType::Pow, XC1x2, XC2x1, Y);
 		UTEST_EQUAL(TEXT("XC1x2^XC2x1[0]"), Y.GetPreparedData<float>()[0], FMath::Pow(1.0f, 3.0f));
 		UTEST_EQUAL(TEXT("XC1x2^XC2x1[1]"), Y.GetPreparedData<float>()[1], FMath::Pow(2.0f, 3.0f));
 		UTEST_EQUAL(TEXT("XC1x2^XC2x1[2]"), Y.GetPreparedData<float>()[2], FMath::Pow(1.0f, 4.0f));
@@ -257,21 +257,21 @@ namespace UE::NNEQA::Private::NNERuntimeRDG::BinaryOp
 
 		//Test output tensor math is correct when both inputs are constant including broadcasting
 		FTensor Y = MakeTensor(TEXT("Y"), { 1 });
-		ElementWiseBinaryCPUHelper::Apply(EElementWiseBinaryOperatorType::Prelu, XC1, XC1, Y);
+		CPUHelper::ElementWiseBinary::Apply(EElementWiseBinaryOperatorType::Prelu, XC1, XC1, Y);
 		UTEST_EQUAL(TEXT("XC1_XC1"), Y.GetPreparedData<float>()[0], Prelu(1.0f, 1.0f));
 
 		Y = MakeTensor(TEXT("Y"), { 1,2 });
-		ElementWiseBinaryCPUHelper::Apply(EElementWiseBinaryOperatorType::Prelu, XC1x2, XC1, Y);
+		CPUHelper::ElementWiseBinary::Apply(EElementWiseBinaryOperatorType::Prelu, XC1x2, XC1, Y);
 		UTEST_EQUAL(TEXT("XC1x2_XC1[0]"), Y.GetPreparedData<float>()[0], Prelu(1.0f, 1.0f));
 		UTEST_EQUAL(TEXT("XC1x2_XC1[1]"), Y.GetPreparedData<float>()[1], Prelu(2.0f, 1.0f));
 
 		Y = MakeTensor(TEXT("Y"), { 1,2 });
-		ElementWiseBinaryCPUHelper::Apply(EElementWiseBinaryOperatorType::Prelu, XC1x2, XC1x2, Y);
+		CPUHelper::ElementWiseBinary::Apply(EElementWiseBinaryOperatorType::Prelu, XC1x2, XC1x2, Y);
 		UTEST_EQUAL(TEXT("XC1x2_XC1x2[0]"), Y.GetPreparedData<float>()[0], Prelu(1.0f, 1.0f));
 		UTEST_EQUAL(TEXT("XC1x2_XC1x2[1]"), Y.GetPreparedData<float>()[1], Prelu(2.0f, 2.0f));
 
 		Y = MakeTensor(TEXT("Y"), { 2,2 });
-		ElementWiseBinaryCPUHelper::Apply(EElementWiseBinaryOperatorType::Prelu, XC1x2, XC2x1, Y);
+		CPUHelper::ElementWiseBinary::Apply(EElementWiseBinaryOperatorType::Prelu, XC1x2, XC2x1, Y);
 		UTEST_EQUAL(TEXT("XC1x2_XC2x1[0]"), Y.GetPreparedData<float>()[0], Prelu(1.0f, 3.0f));
 		UTEST_EQUAL(TEXT("XC1x2_XC2x1[1]"), Y.GetPreparedData<float>()[1], Prelu(2.0f, 3.0f));
 		UTEST_EQUAL(TEXT("XC1x2_XC2x1[2]"), Y.GetPreparedData<float>()[2], Prelu(1.0f, 4.0f));
diff --git a/Engine/Plugins/Experimental/NNE/Source/NNEQA/Private/NNEQAUnitTestOpConcat.cpp b/Engine/Plugins/Experimental/NNE/Source/NNEQA/Private/NNEQAUnitTestOpConcat.cpp
new file mode 100644
--- /dev/null
+++ b/Engine/Plugins/Experimental/NNE/Source/NNEQA/Private/NNEQAUnitTestOpConcat.cpp
@@ -0,0 +1,127 @@
+// Copyright Epic Games, Inc. All Rights Reserved.
+
+#include "CoreMinimal.h"
+#include "NNECoreTensor.h"
+#include "NNECoreTypes.h"
+#include "NNEQAUnitTestHelper.h"
+#include "NNERuntimeRDGHelperConcat.h"
+
+#if WITH_DEV_AUTOMATION_TESTS
+#include "Misc/AutomationTest.h"
+#endif
+
+namespace UE::NNEQA::Private::NNERuntimeRDG::ConcatOp
+{
+
+#if WITH_DEV_AUTOMATION_TESTS
+
+	using namespace NNECore::Internal;
+	using namespace UE::NNERuntimeRDG::Internal;
+
+	IMPLEMENT_NNE_SHAPEINFERENCEHELPER_UNIT_AUTOMATION_TEST(FConcatCPUHelperConstOuput, "System.Engine.MachineLearning.NNE.UnitTest.ConcatHelper.ConstOutput");
+	bool FConcatCPUHelperConstOuput::RunTest(const FString& Parameter)
+	{
+		FTensor XC1 = MakeConstTensor(TEXT("XC1"), { 1 }, { 1.0f });
+		FTensor XC20 = MakeConstTensor(TEXT("XC20"), { 20 }, { 3.0f, 4.0f, 3.0f, 4.0f, 3.0f, 3.0f, 4.0f, 3.0f, 4.0f, 3.0f, 3.0f, 4.0f, 3.0f, 4.0f, 3.0f, 3.0f, 4.0f, 3.0f, 4.0f, 3.0f });
+		FTensor X1 = MakeTensor(TEXT("X"), { 1 });
+
+		FTensor Y = MakeTensor(TEXT("Y"), { 1 });
+		CPUHelper::Concat::Apply({ &XC1 }, Y, 0);
+		UTEST_TRUE(TEXT("Y const if input is const"), Y.HasPreparedData());
+
+		Y = MakeTensor(TEXT("Y"), { 2 });
+		CPUHelper::Concat::Apply({ &XC1, &XC1 }, Y, 0);
+		UTEST_TRUE(TEXT("Y const if inputs are const"), Y.HasPreparedData());
+
+		Y = MakeTensor(TEXT("Y"), { 2 });
+		CPUHelper::Concat::Apply({ &X1, &XC1 }, Y, 0);
+		UTEST_FALSE(TEXT("Y not const if at least on input not const"), Y.HasPreparedData());
+
+		Y = MakeTensor(TEXT("Y"), { 20 });
+		CPUHelper::Concat::Apply({ &XC20 }, Y, 0);
+		UTEST_FALSE(TEXT("Y not const if input is too large "), Y.HasPreparedData());
+
+		return true;
+	}
+
+	IMPLEMENT_NNE_SHAPEINFERENCEHELPER_UNIT_AUTOMATION_TEST(FConcatCPUHelperRank1, "System.Engine.MachineLearning.NNE.UnitTest.ConcatHelper.Rank1");
+	bool FConcatCPUHelperRank1::RunTest(const FString& Parameter)
+	{
+		FTensor XC1 = MakeConstTensor(TEXT("XC1"), { 1 }, { 1.0f });
+		FTensor XC1v2 = MakeConstTensor(TEXT("XC1"), { 1 }, { 2.0f });
+
+		FTensor Y = MakeTensor(TEXT("Y"), { 1 });
+		CPUHelper::Concat::Apply({ &XC1 }, Y, 0);
+		UTEST_TRUE(TEXT("Y const if input is const"), Y.HasPreparedData());
+		UTEST_EQUAL(TEXT("Concat(XC1,XC1)[0]"), Y.GetPreparedData<float>()[0], 1.0f);
+
+		Y = MakeTensor(TEXT("Y"), { 2 });
+		CPUHelper::Concat::Apply({ &XC1, &XC1v2 }, Y, 0);
+		UTEST_TRUE(TEXT("Y const if inputs are const"), Y.HasPreparedData());
+		UTEST_EQUAL(TEXT("Concat(XC1,XC1v2)[0]"), Y.GetPreparedData<float>()[0], 1.0f);
+		UTEST_EQUAL(TEXT("Concat(XC1,XC1v2)[1]"), Y.GetPreparedData<float>()[1], 2.0f);
+
+		Y = MakeTensor(TEXT("Y"), { 3 });
+		CPUHelper::Concat::Apply({ &XC1, &XC1v2, &XC1 }, Y, 0);
+		UTEST_TRUE(TEXT("Y const if inputs are const"), Y.HasPreparedData());
+		UTEST_EQUAL(TEXT("Concat(XC1,XC1v2,XC1)[0]"), Y.GetPreparedData<float>()[0], 1.0f);
+		UTEST_EQUAL(TEXT("Concat(XC1,XC1v2,XC1)[1]"), Y.GetPreparedData<float>()[1], 2.0f);
+		UTEST_EQUAL(TEXT("Concat(XC1,XC1v2,XC1)[2]"), Y.GetPreparedData<float>()[2], 1.0f);
+
+		return true;
+	}
+
+	IMPLEMENT_NNE_SHAPEINFERENCEHELPER_UNIT_AUTOMATION_TEST(FConcatCPUHelperRank3, "System.Engine.MachineLearning.NNE.UnitTest.ConcatHelper.Rank3");
+	bool FConcatCPUHelperRank3::RunTest(const FString& Parameter)
+	{
+		FTensor XC1x2x3 = MakeConstTensor(TEXT("XC1x2x3"), { 1,2,3 }, { 1.0f, 2.0f, 3.0f, 4.0f, 5.0f, 6.0f });
+		FTensor XC1x2x1 = MakeConstTensor(TEXT("XC1x2x1"), { 1,2,1 }, { -1.0f, -2.0f });
+		FTensor XC1x1x3 = MakeConstTensor(TEXT("XC1x1x3"), { 1,1,3 }, { 10.0f, 11.0f, 12.0f });
+
+		FTensor Y = MakeTensor(TEXT("Y"), { 1,2,4 });
+		CPUHelper::Concat::Apply({ &XC1x2x3, &XC1x2x1 }, Y, 2);
+		UTEST_TRUE(TEXT("Y const if inputs are const"), Y.HasPreparedData());
+		UTEST_EQUAL(TEXT("Concat(XC1x2x3,XC1x2x1)[0]"), Y.GetPreparedData<float>()[0], 1.0f);
+		UTEST_EQUAL(TEXT("Concat(XC1x2x3,XC1x2x1)[1]"), Y.GetPreparedData<float>()[1], 2.0f);
+		UTEST_EQUAL(TEXT("Concat(XC1x2x3,XC1x2x1)[2]"), Y.GetPreparedData<float>()[2], 3.0f);
+		UTEST_EQUAL(TEXT("Concat(XC1x2x3,XC1x2x1)[3]"), Y.GetPreparedData<float>()[3], -1.0f);
+		UTEST_EQUAL(TEXT("Concat(XC1x2x3,XC1x2x1)[4]"), Y.GetPreparedData<float>()[4], 4.0f);
+		UTEST_EQUAL(TEXT("Concat(XC1x2x3,XC1x2x1)[5]"), Y.GetPreparedData<float>()[5], 5.0f);
+		UTEST_EQUAL(TEXT("Concat(XC1x2x3,XC1x2x1)[6]"), Y.GetPreparedData<float>()[6], 6.0f);
+		UTEST_EQUAL(TEXT("Concat(XC1x2x3,XC1x2x1)[7]"), Y.GetPreparedData<float>()[7], -2.0f);
+
+		Y = MakeTensor(TEXT("Y"), { 1,3,3 });
+		CPUHelper::Concat::Apply({ &XC1x2x3, &XC1x1x3 }, Y, 1);
+		UTEST_TRUE(TEXT("Y const if inputs are const"), Y.HasPreparedData());
+		UTEST_EQUAL(TEXT("Concat(XC1x2x3,XC1x1x3)[0]"), Y.GetPreparedData<float>()[0], 1.0f);
+		UTEST_EQUAL(TEXT("Concat(XC1x2x3,XC1x1x3)[1]"), Y.GetPreparedData<float>()[1], 2.0f);
+		UTEST_EQUAL(TEXT("Concat(XC1x2x3,XC1x1x3)[2]"), Y.GetPreparedData<float>()[2], 3.0f);
+		UTEST_EQUAL(TEXT("Concat(XC1x2x3,XC1x1x3)[3]"), Y.GetPreparedData<float>()[3], 4.0f);
+		UTEST_EQUAL(TEXT("Concat(XC1x2x3,XC1x1x3)[4]"), Y.GetPreparedData<float>()[4], 5.0f);
+		UTEST_EQUAL(TEXT("Concat(XC1x2x3,XC1x1x3)[5]"), Y.GetPreparedData<float>()[5], 6.0f);
+		UTEST_EQUAL(TEXT("Concat(XC1x2x3,XC1x1x3)[6]"), Y.GetPreparedData<float>()[6], 10.0f);
+		UTEST_EQUAL(TEXT("Concat(XC1x2x3,XC1x1x3)[7]"), Y.GetPreparedData<float>()[7], 11.0f);
+		UTEST_EQUAL(TEXT("Concat(XC1x2x3,XC1x1x3)[8]"), Y.GetPreparedData<float>()[8], 12.0f);
+
+		Y = MakeTensor(TEXT("Y"), { 2,2,3 });
+		CPUHelper::Concat::Apply({ &XC1x2x3, &XC1x2x3 }, Y, 0);
+		UTEST_TRUE(TEXT("Y const if inputs are const"), Y.HasPreparedData());
+		UTEST_EQUAL(TEXT("Concat(XC1x2x3,XC1x2x3)[0]"), Y.GetPreparedData<float>()[0], 1.0f);
+		UTEST_EQUAL(TEXT("Concat(XC1x2x3,XC1x2x3)[1]"), Y.GetPreparedData<float>()[1], 2.0f);
+		UTEST_EQUAL(TEXT("Concat(XC1x2x3,XC1x2x3)[2]"), Y.GetPreparedData<float>()[2], 3.0f);
+		UTEST_EQUAL(TEXT("Concat(XC1x2x3,XC1x2x3)[3]"), Y.GetPreparedData<float>()[3], 4.0f);
+		UTEST_EQUAL(TEXT("Concat(XC1x2x3,XC1x2x3)[4]"), Y.GetPreparedData<float>()[4], 5.0f);
+		UTEST_EQUAL(TEXT("Concat(XC1x2x3,XC1x2x3)[5]"), Y.GetPreparedData<float>()[5], 6.0f);
+		UTEST_EQUAL(TEXT("Concat(XC1x2x3,XC1x2x3)[6]"), Y.GetPreparedData<float>()[6], 1.0f);
+		UTEST_EQUAL(TEXT("Concat(XC1x2x3,XC1x2x3)[7]"), Y.GetPreparedData<float>()[7], 2.0f);
+		UTEST_EQUAL(TEXT("Concat(XC1x2x3,XC1x2x3)[8]"), Y.GetPreparedData<float>()[8], 3.0f);
+		UTEST_EQUAL(TEXT("Concat(XC1x2x3,XC1x2x3)[9]"), Y.GetPreparedData<float>()[9], 4.0f);
+		UTEST_EQUAL(TEXT("Concat(XC1x2x3,XC1x2x3)[10]"), Y.GetPreparedData<float>()[10], 5.0f);
+		UTEST_EQUAL(TEXT("Concat(XC1x2x3,XC1x2x3)[11]"), Y.GetPreparedData<float>()[11], 6.0f);
+
+		return true;
+	}
+
+#endif //WITH_DEV_AUTOMATION_TESTS
+
+} // namespace UE::NNEQA::Private::NNERuntimeRDG::ConcatOp
diff --git a/Engine/Plugins/Experimental/NNE/Source/NNEQA/Private/NNEQAUnitTestOpSlice.cpp b/Engine/Plugins/Experimental/NNE/Source/NNEQA/Private/NNEQAUnitTestOpSlice.cpp
new file mode 100644
--- /dev/null
+++ b/Engine/Plugins/Experimental/NNE/Source/NNEQA/Private/NNEQAUnitTestOpSlice.cpp
@@ -0,0 +1,108 @@
+// Copyright Epic Games, Inc. All Rights Reserved.
+
+#include "CoreMinimal.h"
+#include "NNECoreTensor.h"
+#include "NNECoreTypes.h"
+#include "NNEQAUnitTestHelper.h"
+#include "NNERuntimeRDGHelperSlice.h"
+
+#if WITH_DEV_AUTOMATION_TESTS
+#include "Misc/AutomationTest.h"
+#endif
+
+namespace UE::NNEQA::Private::NNERuntimeRDG::SliceOp
+{
+
+#if WITH_DEV_AUTOMATION_TESTS
+
+	using namespace NNECore::Internal;
+	using namespace UE::NNERuntimeRDG::Internal;
+
+	IMPLEMENT_NNE_SHAPEINFERENCEHELPER_UNIT_AUTOMATION_TEST(FSliceCPUHelperConstOuput, "System.Engine.MachineLearning.NNE.UnitTest.SliceHelper.ConstOutput");
+	bool FSliceCPUHelperConstOuput::RunTest(const FString& Parameter)
+	{
+		FTensor XC1 = MakeConstTensor(TEXT("XC1"), { 1 }, { 1.0f });
+		FTensor XC20 = MakeConstTensor(TEXT("XC20"), { 20 }, { 3.0f, 4.0f, 3.0f, 4.0f, 3.0f, 3.0f, 4.0f, 3.0f, 4.0f, 3.0f, 3.0f, 4.0f, 3.0f, 4.0f, 3.0f, 3.0f, 4.0f, 3.0f, 4.0f, 3.0f });
+		FTensor X1 = MakeTensor(TEXT("X"), { 1 });
+
+		FTensor Y = MakeTensor(TEXT("Y"), { 1 });
+		CPUHelper::Slice::Apply(XC1 , Y, { 0 });
+		UTEST_TRUE(TEXT("Y const if input is const"), Y.HasPreparedData());
+
+		Y = MakeTensor(TEXT("Y"), { 1 });
+		CPUHelper::Slice::Apply(X1, Y, { 0 });
+		UTEST_FALSE(TEXT("Y not const if input not const"), Y.HasPreparedData());
+
+		Y = MakeTensor(TEXT("Y"), { 20 });
+		CPUHelper::Slice::Apply(XC20, Y, { 0 });
+		UTEST_FALSE(TEXT("Y not const if input is too large "), Y.HasPreparedData());
+
+		return true;
+	}
+
+	IMPLEMENT_NNE_SHAPEINFERENCEHELPER_UNIT_AUTOMATION_TEST(FSliceCPUHelperRank1, "System.Engine.MachineLearning.NNE.UnitTest.SliceHelper.Rank1");
+	bool FSliceCPUHelperRank1::RunTest(const FString& Parameter)
+	{
+		FTensor XC6 = MakeConstTensor(TEXT("XC6"), { 6 }, { 1.0f, 2.0f, 3.0f, 4.0f, 5.0f, 6.0f });
+
+		FTensor Y = MakeTensor(TEXT("Y"), { 6 });
+		CPUHelper::Slice::Apply(XC6, Y, { 0 });
+		UTEST_TRUE(TEXT("Y const if input is const"), Y.HasPreparedData());
+		UTEST_EQUAL(TEXT("Slice(XC6,6,0)[0]"), Y.GetPreparedData<float>()[0], 1.0f);
+		UTEST_EQUAL(TEXT("Slice(XC6,6,0)[1]"), Y.GetPreparedData<float>()[1], 2.0f);
+		UTEST_EQUAL(TEXT("Slice(XC6,6,0)[2]"), Y.GetPreparedData<float>()[2], 3.0f);
+		UTEST_EQUAL(TEXT("Slice(XC6,6,0)[3]"), Y.GetPreparedData<float>()[3], 4.0f);
+		UTEST_EQUAL(TEXT("Slice(XC6,6,0)[4]"), Y.GetPreparedData<float>()[4], 5.0f);
+		UTEST_EQUAL(TEXT("Slice(XC6,6,0)[5]"), Y.GetPreparedData<float>()[5], 6.0f);
+
+		Y = MakeTensor(TEXT("Y"), { 2 });
+		CPUHelper::Slice::Apply(XC6, Y, { 4 });
+		UTEST_TRUE(TEXT("Y const if input is const"), Y.HasPreparedData());
+		UTEST_EQUAL(TEXT("Slice(XC6,2,4)[0]"), Y.GetPreparedData<float>()[0], 5.0f);
+		UTEST_EQUAL(TEXT("Slice(XC6,2,4)[1]"), Y.GetPreparedData<float>()[1], 6.0f);
+
+		Y = MakeTensor(TEXT("Y"), { 2 });
+		CPUHelper::Slice::Apply(XC6, Y, { 1 });
+		UTEST_TRUE(TEXT("Y const if input is const"), Y.HasPreparedData());
+		UTEST_EQUAL(TEXT("Slice(XC6,2,1)[0]"), Y.GetPreparedData<float>()[0], 2.0f);
+		UTEST_EQUAL(TEXT("Slice(XC6,2,1)[1]"), Y.GetPreparedData<float>()[1], 3.0f);
+
+		return true;
+	}
+
+	IMPLEMENT_NNE_SHAPEINFERENCEHELPER_UNIT_AUTOMATION_TEST(FSliceCPUHelperRank3, "System.Engine.MachineLearning.NNE.UnitTest.SliceHelper.Rank3");
+	bool FSliceCPUHelperRank3::RunTest(const FString& Parameter)
+	{
+		FTensor XC1x2x3 = MakeConstTensor(TEXT("XC1x2x3"), { 1,2,3 }, { 1.0f, 2.0f, 3.0f, 4.0f, 5.0f, 6.0f });
+
+		FTensor Y = MakeTensor(TEXT("Y"), { 1,2,3 });
+		CPUHelper::Slice::Apply(XC1x2x3, Y, { 0,0,0 });
+		UTEST_TRUE(TEXT("Y const if input is const"), Y.HasPreparedData());
+		UTEST_EQUAL(TEXT("Slice(XC1x2x3,1x2x3,0-0-0)[0]"), Y.GetPreparedData<float>()[0], 1.0f);
+		UTEST_EQUAL(TEXT("Slice(XC1x2x3,1x2x3,0-0-0)[1]"), Y.GetPreparedData<float>()[1], 2.0f);
+		UTEST_EQUAL(TEXT("Slice(XC1x2x3,1x2x3,0-0-0)[2]"), Y.GetPreparedData<float>()[2], 3.0f);
+		UTEST_EQUAL(TEXT("Slice(XC1x2x3,1x2x3,0-0-0)[3]"), Y.GetPreparedData<float>()[3], 4.0f);
+		UTEST_EQUAL(TEXT("Slice(XC1x2x3,1x2x3,0-0-0)[4]"), Y.GetPreparedData<float>()[4], 5.0f);
+		UTEST_EQUAL(TEXT("Slice(XC1x2x3,1x2x3,0-0-0)[5]"), Y.GetPreparedData<float>()[5], 6.0f);
+
+		Y = MakeTensor(TEXT("Y"), { 1,1,3 });
+		CPUHelper::Slice::Apply(XC1x2x3, Y, { 0,1,0 });
+		UTEST_TRUE(TEXT("Y const if input is const"), Y.HasPreparedData());
+		UTEST_EQUAL(TEXT("Slice(XC1x2x3,1x1x3,0-1-0)[0]"), Y.GetPreparedData<float>()[0], 4.0f);
+		UTEST_EQUAL(TEXT("Slice(XC1x2x3,1x1x3,0-1-0)[1]"), Y.GetPreparedData<float>()[1], 5.0f);
+		UTEST_EQUAL(TEXT("Slice(XC1x2x3,1x1x3,0-1-0)[2]"), Y.GetPreparedData<float>()[2], 6.0f);
+
+		Y = MakeTensor(TEXT("Y"), { 1,2,2 });
+		CPUHelper::Slice::Apply(XC1x2x3, Y, { 0,0,1 });
+		UTEST_TRUE(TEXT("Y const if input is const"), Y.HasPreparedData());
+		UTEST_EQUAL(TEXT("Slice(XC1x2x3,1x2x2,0-0-1)[0]"), Y.GetPreparedData<float>()[0], 2.0f);
+		UTEST_EQUAL(TEXT("Slice(XC1x2x3,1x2x2,0-0-1)[1]"), Y.GetPreparedData<float>()[1], 3.0f);
+		UTEST_EQUAL(TEXT("Slice(XC1x2x3,1x2x2,0-0-1)[2]"), Y.GetPreparedData<float>()[2], 5.0f);
+		UTEST_EQUAL(TEXT("Slice(XC1x2x3,1x2x2,0-0-1)[3]"), Y.GetPreparedData<float>()[3], 6.0f);
+
+		return true;
+	}
+
+#endif //WITH_DEV_AUTOMATION_TESTS
+
+} // namespace UE::NNEQA::Private::NNERuntimeRDG::SliceOp
diff --git a/Engine/Plugins/Experimental/NNE/Source/NNEQA/Private/NNEQAUnitTestOpUnary.cpp b/Engine/Plugins/Experimental/NNE/Source/NNEQA/Private/NNEQAUnitTestOpUnary.cpp
--- a/Engine/Plugins/Experimental/NNE/Source/NNEQA/Private/NNEQAUnitTestOpUnary.cpp
+++ b/Engine/Plugins/Experimental/NNE/Source/NNEQA/Private/NNEQAUnitTestOpUnary.cpp
@@ -5,7 +5,7 @@
 #include "NNECoreTensor.h"
 #include "NNECoreTypes.h"
 #include "NNEQAUnitTestHelper.h"
-#include "NNERuntimeRDGElementWiseUnaryHelper.h"
+#include "NNERuntimeRDGHelperElementWiseUnary.h"
 
 #if WITH_DEV_AUTOMATION_TESTS
 #include "Misc/AutomationTest.h"
@@ -40,7 +40,8 @@ namespace UE::NNEQA::Private::NNERuntimeRDG::UnaryOp
 	{
 		FTensor XC2x1 = MakeConstTensor(TEXT("XC2x1"), { 2,1 }, { -1.0f, 2.0f });
 		FTensor Y = MakeTensor(TEXT("Y"), { 2,1 });
-		ElementWiseUnaryCPUHelper::Apply(EElementWiseUnaryOperatorType::Abs, XC2x1, 1.0f, 2.0f, 3.0f, Y);
+		CPUHelper::ElementWiseUnary::Apply(EElementWiseUnaryOperatorType::Abs, XC2x1, 1.0f, 2.0f, 3.0f, Y);
+		UTEST_TRUE(TEXT("Y const if input is const"), Y.HasPreparedData());
 		UTEST_EQUAL(TEXT("Abs(XC2x1)[0]"), Y.GetPreparedData<float>()[0], FMath::Abs(XC2x1.GetPreparedData<float>()[0]));
 		UTEST_EQUAL(TEXT("Abs(XC2x1)[1]"), Y.GetPreparedData<float>()[1], FMath::Abs(XC2x1.GetPreparedData<float>()[1]));
 		return true;
@@ -51,7 +52,8 @@ namespace UE::NNEQA::Private::NNERuntimeRDG::UnaryOp
 	{
 		FTensor XC2x1 = MakeConstTensor(TEXT("XC2x1"), { 2,1 }, { -0.5f, 0.8f });
 		FTensor Y = MakeTensor(TEXT("Y"), { 2,1 });
-		ElementWiseUnaryCPUHelper::Apply(EElementWiseUnaryOperatorType::Acos, XC2x1, 1.0f, 2.0f, 3.0f, Y);
+		CPUHelper::ElementWiseUnary::Apply(EElementWiseUnaryOperatorType::Acos, XC2x1, 1.0f, 2.0f, 3.0f, Y);
+		UTEST_TRUE(TEXT("Y const if input is const"), Y.HasPreparedData());
 		UTEST_EQUAL(TEXT("Acos(XC2x1)[0]"), Y.GetPreparedData<float>()[0], FMath::Acos(XC2x1.GetPreparedData<float>()[0]));
 		UTEST_EQUAL(TEXT("Acos(XC2x1)[1]"), Y.GetPreparedData<float>()[1], FMath::Acos(XC2x1.GetPreparedData<float>()[1]));
 		return true;
@@ -62,7 +64,8 @@ namespace UE::NNEQA::Private::NNERuntimeRDG::UnaryOp
 	{
 		FTensor XC2x1 = MakeConstTensor(TEXT("XC2x1"), { 2,1 }, { 1.0f, 1.5f });
 		FTensor Y = MakeTensor(TEXT("Y"), { 2,1 });
-		ElementWiseUnaryCPUHelper::Apply(EElementWiseUnaryOperatorType::Acosh, XC2x1, 1.0f, 2.0f, 3.0f, Y);
+		CPUHelper::ElementWiseUnary::Apply(EElementWiseUnaryOperatorType::Acosh, XC2x1, 1.0f, 2.0f, 3.0f, Y);
+		UTEST_TRUE(TEXT("Y const if input is const"), Y.HasPreparedData());
 		UTEST_EQUAL(TEXT("Acosh(XC2x1)[0]"), Y.GetPreparedData<float>()[0], 0.0f);
 		UTEST_EQUAL(TEXT("Acosh(XC2x1)[1]"), Y.GetPreparedData<float>()[1], 0.9624236501192069f);
 		return true;
@@ -73,7 +76,8 @@ namespace UE::NNEQA::Private::NNERuntimeRDG::UnaryOp
 	{
 		FTensor XC2x1 = MakeConstTensor(TEXT("XC2x1"), { 2,1 }, { -0.5f, 0.8f });
 		FTensor Y = MakeTensor(TEXT("Y"), { 2,1 });
-		ElementWiseUnaryCPUHelper::Apply(EElementWiseUnaryOperatorType::Asin, XC2x1, 1.0f, 2.0f, 3.0f, Y);
+		CPUHelper::ElementWiseUnary::Apply(EElementWiseUnaryOperatorType::Asin, XC2x1, 1.0f, 2.0f, 3.0f, Y);
+		UTEST_TRUE(TEXT("Y const if input is const"), Y.HasPreparedData());
 		UTEST_EQUAL(TEXT("Asin(XC2x1)[0]"), Y.GetPreparedData<float>()[0], FMath::Asin(XC2x1.GetPreparedData<float>()[0]));
 		UTEST_EQUAL(TEXT("Asin(XC2x1)[1]"), Y.GetPreparedData<float>()[1], FMath::Asin(XC2x1.GetPreparedData<float>()[1]));
 		return true;
@@ -84,7 +88,8 @@ namespace UE::NNEQA::Private::NNERuntimeRDG::UnaryOp
 	{
 		FTensor XC2x1 = MakeConstTensor(TEXT("XC2x1"), { 2,1 }, { 1.0f, 2.0f });
 		FTensor Y = MakeTensor(TEXT("Y"), { 2,1 });
-		ElementWiseUnaryCPUHelper::Apply(EElementWiseUnaryOperatorType::Asinh, XC2x1, 1.0f, 2.0f, 3.0f, Y);
+		CPUHelper::ElementWiseUnary::Apply(EElementWiseUnaryOperatorType::Asinh, XC2x1, 1.0f, 2.0f, 3.0f, Y);
+		UTEST_TRUE(TEXT("Y const if input is const"), Y.HasPreparedData());
 		UTEST_EQUAL(TEXT("Asinh(XC2x1)[0]"), Y.GetPreparedData<float>()[0], 0.881373587019543f);
 		UTEST_EQUAL(TEXT("Asinh(XC2x1)[1]"), Y.GetPreparedData<float>()[1], 1.4436354751788103f);
 		return true;
@@ -95,7 +100,8 @@ namespace UE::NNEQA::Private::NNERuntimeRDG::UnaryOp
 	{
 		FTensor XC2x1 = MakeConstTensor(TEXT("XC2x1"), { 2,1 }, { -0.5f, 0.8f });
 		FTensor Y = MakeTensor(TEXT("Y"), { 2,1 });
-		ElementWiseUnaryCPUHelper::Apply(EElementWiseUnaryOperatorType::Atan, XC2x1, 1.0f, 2.0f, 3.0f, Y);
+		CPUHelper::ElementWiseUnary::Apply(EElementWiseUnaryOperatorType::Atan, XC2x1, 1.0f, 2.0f, 3.0f, Y);
+		UTEST_TRUE(TEXT("Y const if input is const"), Y.HasPreparedData());
 		UTEST_EQUAL(TEXT("Atan(XC2x1)[0]"), Y.GetPreparedData<float>()[0], FMath::Atan(XC2x1.GetPreparedData<float>()[0]));
 		UTEST_EQUAL(TEXT("Atan(XC2x1)[1]"), Y.GetPreparedData<float>()[1], FMath::Atan(XC2x1.GetPreparedData<float>()[1]));
 		return true;
@@ -106,7 +112,8 @@ namespace UE::NNEQA::Private::NNERuntimeRDG::UnaryOp
 	{
 		FTensor XC2x1 = MakeConstTensor(TEXT("XC2x1"), { 2,1 }, { -0.5f, 0.8f });
 		FTensor Y = MakeTensor(TEXT("Y"), { 2,1 });
-		ElementWiseUnaryCPUHelper::Apply(EElementWiseUnaryOperatorType::Ceil, XC2x1, 1.0f, 2.0f, 3.0f, Y);
+		CPUHelper::ElementWiseUnary::Apply(EElementWiseUnaryOperatorType::Ceil, XC2x1, 1.0f, 2.0f, 3.0f, Y);
+		UTEST_TRUE(TEXT("Y const if input is const"), Y.HasPreparedData());
 		UTEST_EQUAL(TEXT("Ceil(XC2x1)[0]"), Y.GetPreparedData<float>()[0], FMath::CeilToFloat(XC2x1.GetPreparedData<float>()[0]));
 		UTEST_EQUAL(TEXT("Ceil(XC2x1)[1]"), Y.GetPreparedData<float>()[1], FMath::CeilToFloat(XC2x1.GetPreparedData<float>()[1]));
 		return true;
@@ -116,7 +123,8 @@ namespace UE::NNEQA::Private::NNERuntimeRDG::UnaryOp
 	{
 		FTensor XC2x1 = MakeConstTensor(TEXT("XC2x1"), { 2,1 }, { -0.5f, 0.8f });
 		FTensor Y = MakeTensor(TEXT("Y"), { 2,1 });
-		ElementWiseUnaryCPUHelper::Apply(EElementWiseUnaryOperatorType::Cos, XC2x1, 1.0f, 2.0f, 3.0f, Y);
+		CPUHelper::ElementWiseUnary::Apply(EElementWiseUnaryOperatorType::Cos, XC2x1, 1.0f, 2.0f, 3.0f, Y);
+		UTEST_TRUE(TEXT("Y const if input is const"), Y.HasPreparedData());
 		UTEST_EQUAL(TEXT("Cos(XC2x1)[0]"), Y.GetPreparedData<float>()[0], FMath::Cos(XC2x1.GetPreparedData<float>()[0]));
 		UTEST_EQUAL(TEXT("Cos(XC2x1)[1]"), Y.GetPreparedData<float>()[1], FMath::Cos(XC2x1.GetPreparedData<float>()[1]));
 		return true;
@@ -127,7 +135,8 @@ namespace UE::NNEQA::Private::NNERuntimeRDG::UnaryOp
 	{
 		FTensor XC2x1 = MakeConstTensor(TEXT("XC2x1"), { 2,1 }, { 1.0f, 2.0f });
 		FTensor Y = MakeTensor(TEXT("Y"), { 2,1 });
-		ElementWiseUnaryCPUHelper::Apply(EElementWiseUnaryOperatorType::Cosh, XC2x1, 1.0f, 2.0f, 3.0f, Y);
+		CPUHelper::ElementWiseUnary::Apply(EElementWiseUnaryOperatorType::Cosh, XC2x1, 1.0f, 2.0f, 3.0f, Y);
+		UTEST_TRUE(TEXT("Y const if input is const"), Y.HasPreparedData());
 		UTEST_EQUAL(TEXT("Cosh(XC2x1)[0]"), Y.GetPreparedData<float>()[0], 1.5430806348152437f);
 		UTEST_EQUAL(TEXT("Cosh(XC2x1)[1]"), Y.GetPreparedData<float>()[1], 3.7621956910836314f);
 		return true;
@@ -137,7 +146,8 @@ namespace UE::NNEQA::Private::NNERuntimeRDG::UnaryOp
 	{
 		FTensor XC2x1 = MakeConstTensor(TEXT("XC2x1"), { 2,1 }, { -0.5f, 0.8f });
 		FTensor Y = MakeTensor(TEXT("Y"), { 2,1 });
-		ElementWiseUnaryCPUHelper::Apply(EElementWiseUnaryOperatorType::Elu, XC2x1, 1.0f, 2.0f, 3.0f, Y);
+		CPUHelper::ElementWiseUnary::Apply(EElementWiseUnaryOperatorType::Elu, XC2x1, 1.0f, 2.0f, 3.0f, Y);
+		UTEST_TRUE(TEXT("Y const if input is const"), Y.HasPreparedData());
 		UTEST_EQUAL(TEXT("Elu(XC2x1)[0]"), Y.GetPreparedData<float>()[0], -0.3934693402873666f);
 		UTEST_EQUAL(TEXT("Elu(XC2x1)[1]"), Y.GetPreparedData<float>()[1], 0.8f);
 		return true;
@@ -148,7 +158,8 @@ namespace UE::NNEQA::Private::NNERuntimeRDG::UnaryOp
 	{
 		FTensor XC2x1 = MakeConstTensor(TEXT("XC2x1"), { 2,1 }, { 1.0f, 2.0f });
 		FTensor Y = MakeTensor(TEXT("Y"), { 2,1 });
-		ElementWiseUnaryCPUHelper::Apply(EElementWiseUnaryOperatorType::Erf, XC2x1, 1.0f, 2.0f, 3.0f, Y);
+		CPUHelper::ElementWiseUnary::Apply(EElementWiseUnaryOperatorType::Erf, XC2x1, 1.0f, 2.0f, 3.0f, Y);
+		UTEST_TRUE(TEXT("Y const if input is const"), Y.HasPreparedData());
 		UTEST_EQUAL_TOLERANCE(TEXT("Erf(XC2x1)[0]"), Y.GetPreparedData<float>()[0], 0.8427007929497149f, 1e-3f);
 		UTEST_EQUAL_TOLERANCE(TEXT("Erf(XC2x1)[1]"), Y.GetPreparedData<float>()[1], 0.9953222650189527f, 1e-3f);
 		return true;
@@ -159,7 +170,8 @@ namespace UE::NNEQA::Private::NNERuntimeRDG::UnaryOp
 	{
 		FTensor XC2x1 = MakeConstTensor(TEXT("XC2x1"), { 2,1 }, { -0.5f, 0.8f });
 		FTensor Y = MakeTensor(TEXT("Y"), { 2,1 });
-		ElementWiseUnaryCPUHelper::Apply(EElementWiseUnaryOperatorType::Exp, XC2x1, 1.0f, 2.0f, 3.0f, Y);
+		CPUHelper::ElementWiseUnary::Apply(EElementWiseUnaryOperatorType::Exp, XC2x1, 1.0f, 2.0f, 3.0f, Y);
+		UTEST_TRUE(TEXT("Y const if input is const"), Y.HasPreparedData());
 		UTEST_EQUAL(TEXT("Exp(XC2x1)[0]"), Y.GetPreparedData<float>()[0], FMath::Exp(XC2x1.GetPreparedData<float>()[0]));
 		UTEST_EQUAL(TEXT("Exp(XC2x1)[1]"), Y.GetPreparedData<float>()[1], FMath::Exp(XC2x1.GetPreparedData<float>()[1]));
 		return true;
@@ -170,7 +182,8 @@ namespace UE::NNEQA::Private::NNERuntimeRDG::UnaryOp
 	{
 		FTensor XC2x1 = MakeConstTensor(TEXT("XC2x1"), { 2,1 }, { -0.5f, 0.8f });
 		FTensor Y = MakeTensor(TEXT("Y"), { 2,1 });
-		ElementWiseUnaryCPUHelper::Apply(EElementWiseUnaryOperatorType::Floor, XC2x1, 1.0f, 2.0f, 3.0f, Y);
+		CPUHelper::ElementWiseUnary::Apply(EElementWiseUnaryOperatorType::Floor, XC2x1, 1.0f, 2.0f, 3.0f, Y);
+		UTEST_TRUE(TEXT("Y const if input is const"), Y.HasPreparedData());
 		UTEST_EQUAL(TEXT("Floor(XC2x1)[0]"), Y.GetPreparedData<float>()[0], FMath::Floor(XC2x1.GetPreparedData<float>()[0]));
 		UTEST_EQUAL(TEXT("Floor(XC2x1)[1]"), Y.GetPreparedData<float>()[1], FMath::Floor(XC2x1.GetPreparedData<float>()[1]));
 		return true;
@@ -181,7 +194,8 @@ namespace UE::NNEQA::Private::NNERuntimeRDG::UnaryOp
 	{
 		FTensor XC2x1 = MakeConstTensor(TEXT("XC2x1"), { 2,1 }, { -0.5f, INFINITY });
 		FTensor Y = MakeTensor(TEXT("Y"), { 2,1 });
-		ElementWiseUnaryCPUHelper::Apply(EElementWiseUnaryOperatorType::IsInf, XC2x1, 1.0f, 2.0f, 3.0f, Y);
+		CPUHelper::ElementWiseUnary::Apply(EElementWiseUnaryOperatorType::IsInf, XC2x1, 1.0f, 2.0f, 3.0f, Y);
+		UTEST_TRUE(TEXT("Y const if input is const"), Y.HasPreparedData());
 		UTEST_EQUAL(TEXT("IsInf(XC2x1)[0]"), Y.GetPreparedData<float>()[0], 0.0f);
 		UTEST_EQUAL(TEXT("IsInf(XC2x1)[1]"), Y.GetPreparedData<float>()[1], 1.0f);
 		return true;
@@ -193,7 +207,8 @@ namespace UE::NNEQA::Private::NNERuntimeRDG::UnaryOp
 		float FloatNan = FMath::Sqrt(-1.0f);
 		FTensor XC2x1 = MakeConstTensor(TEXT("XC2x1"), { 2,1 }, { -0.5f, FloatNan });
 		FTensor Y = MakeTensor(TEXT("Y"), { 2,1 });
-		ElementWiseUnaryCPUHelper::Apply(EElementWiseUnaryOperatorType::IsNan, XC2x1, 1.0f, 2.0f, 3.0f, Y);
+		CPUHelper::ElementWiseUnary::Apply(EElementWiseUnaryOperatorType::IsNan, XC2x1, 1.0f, 2.0f, 3.0f, Y);
+		UTEST_TRUE(TEXT("Y const if input is const"), Y.HasPreparedData());
 		UTEST_EQUAL(TEXT("IsNan(XC2x1)[0]"), Y.GetPreparedData<float>()[0], 0.0f);
 		UTEST_EQUAL(TEXT("IsNan(XC2x1)[1]"), Y.GetPreparedData<float>()[1], 1.0f);
 		return true;
@@ -204,7 +219,8 @@ namespace UE::NNEQA::Private::NNERuntimeRDG::UnaryOp
 	{
 		FTensor XC2x1 = MakeConstTensor(TEXT("XC2x1"), { 2,1 }, { 0.0f, -1.5f });
 		FTensor Y = MakeTensor(TEXT("Y"), { 2,1 });
-		ElementWiseUnaryCPUHelper::Apply(EElementWiseUnaryOperatorType::HardSigmoid, XC2x1, 1.0f, 2.0f, 3.0f, Y);
+		CPUHelper::ElementWiseUnary::Apply(EElementWiseUnaryOperatorType::HardSigmoid, XC2x1, 1.0f, 2.0f, 3.0f, Y);
+		UTEST_TRUE(TEXT("Y const if input is const"), Y.HasPreparedData());
 		UTEST_EQUAL(TEXT("HardSigmoid(XC2x1)[0]"), Y.GetPreparedData<float>()[0], 1.0f);
 		UTEST_EQUAL(TEXT("HardSigmoid(XC2x1)[1]"), Y.GetPreparedData<float>()[1], 0.5f);
 		return true;
@@ -215,7 +231,8 @@ namespace UE::NNEQA::Private::NNERuntimeRDG::UnaryOp
 	{
 		FTensor XC2x1 = MakeConstTensor(TEXT("XC2x1"), { 2,1 }, { 0.0f, 1.0f });
 		FTensor Y = MakeTensor(TEXT("Y"), { 2,1 });
-		ElementWiseUnaryCPUHelper::Apply(EElementWiseUnaryOperatorType::HardSwish, XC2x1, 1.0f, 2.0f, 3.0f, Y);
+		CPUHelper::ElementWiseUnary::Apply(EElementWiseUnaryOperatorType::HardSwish, XC2x1, 1.0f, 2.0f, 3.0f, Y);
+		UTEST_TRUE(TEXT("Y const if input is const"), Y.HasPreparedData());
 		UTEST_EQUAL(TEXT("HardSwitch(XC2x1)[0]"), Y.GetPreparedData<float>()[0], 0.5f);
 		UTEST_EQUAL(TEXT("HardSwitch(XC2x1)[1]"), Y.GetPreparedData<float>()[1], (1.0f / 6.0f) + 0.5f);
 		return true;
@@ -226,7 +243,8 @@ namespace UE::NNEQA::Private::NNERuntimeRDG::UnaryOp
 	{
 		FTensor XC2x1 = MakeConstTensor(TEXT("XC2x1"), { 2,1 }, { 1.0f, -1.0f });
 		FTensor Y = MakeTensor(TEXT("Y"), { 2,1 });
-		ElementWiseUnaryCPUHelper::Apply(EElementWiseUnaryOperatorType::LeakyRelu, XC2x1, -0.5f, 2.0f, 3.0f, Y);
+		CPUHelper::ElementWiseUnary::Apply(EElementWiseUnaryOperatorType::LeakyRelu, XC2x1, -0.5f, 2.0f, 3.0f, Y);
+		UTEST_TRUE(TEXT("Y const if input is const"), Y.HasPreparedData());
 		UTEST_EQUAL(TEXT("LeakyRelu(XC2x1)[0]"), Y.GetPreparedData<float>()[0], 1.0f);
 		UTEST_EQUAL(TEXT("LeakyRelu(XC2x1)[1]"), Y.GetPreparedData<float>()[1], 0.5f);
 		return true;
@@ -237,7 +255,8 @@ namespace UE::NNEQA::Private::NNERuntimeRDG::UnaryOp
 	{
 		FTensor XC2x1 = MakeConstTensor(TEXT("XC2x1"), { 2,1 }, { -0.5f, 0.8f });
 		FTensor Y = MakeTensor(TEXT("Y"), { 2,1 });
-		ElementWiseUnaryCPUHelper::Apply(EElementWiseUnaryOperatorType::Log, XC2x1, 1.0f, 2.0f, 3.0f, Y);
+		CPUHelper::ElementWiseUnary::Apply(EElementWiseUnaryOperatorType::Log, XC2x1, 1.0f, 2.0f, 3.0f, Y);
+		UTEST_TRUE(TEXT("Y const if input is const"), Y.HasPreparedData());
 		UTEST_EQUAL(TEXT("Log(XC2x1)[0]"), Y.GetPreparedData<float>()[0], FMath::Loge(XC2x1.GetPreparedData<float>()[0]));
 		UTEST_EQUAL(TEXT("Log(XC2x1)[1]"), Y.GetPreparedData<float>()[1], FMath::Loge(XC2x1.GetPreparedData<float>()[1]));
 		return true;
@@ -248,7 +267,8 @@ namespace UE::NNEQA::Private::NNERuntimeRDG::UnaryOp
 	{
 		FTensor XC2x1 = MakeConstTensor(TEXT("XC2x1"), { 2,1 }, { -0.5f, 0.8f });
 		FTensor Y = MakeTensor(TEXT("Y"), { 2,1 });
-		ElementWiseUnaryCPUHelper::Apply(EElementWiseUnaryOperatorType::Neg, XC2x1, 1.0f, 2.0f, 3.0f, Y);
+		CPUHelper::ElementWiseUnary::Apply(EElementWiseUnaryOperatorType::Neg, XC2x1, 1.0f, 2.0f, 3.0f, Y);
+		UTEST_TRUE(TEXT("Y const if input is const"), Y.HasPreparedData());
 		UTEST_EQUAL(TEXT("Neg(XC2x1)[0]"), Y.GetPreparedData<float>()[0], -XC2x1.GetPreparedData<float>()[0]);
 		UTEST_EQUAL(TEXT("Neg(XC2x1)[1]"), Y.GetPreparedData<float>()[1], -XC2x1.GetPreparedData<float>()[1]);
 		return true;
@@ -259,7 +279,8 @@ namespace UE::NNEQA::Private::NNERuntimeRDG::UnaryOp
 	{
 		FTensor XC2x1 = MakeConstTensor(TEXT("XC2x1"), { 2,1 }, { -0.5f, 0.8f });
 		FTensor Y = MakeTensor(TEXT("Y"), { 2,1 });
-		ElementWiseUnaryCPUHelper::Apply(EElementWiseUnaryOperatorType::Reciprocal, XC2x1, 1.0f, 2.0f, 3.0f, Y);
+		CPUHelper::ElementWiseUnary::Apply(EElementWiseUnaryOperatorType::Reciprocal, XC2x1, 1.0f, 2.0f, 3.0f, Y);
+		UTEST_TRUE(TEXT("Y const if input is const"), Y.HasPreparedData());
 		UTEST_EQUAL(TEXT("Reciprocal(XC2x1)[0]"), Y.GetPreparedData<float>()[0], 1.0f / (XC2x1.GetPreparedData<float>()[0]));
 		UTEST_EQUAL(TEXT("Reciprocal(XC2x1)[1]"), Y.GetPreparedData<float>()[1], 1.0f / (XC2x1.GetPreparedData<float>()[1]));
 		return true;
@@ -270,7 +291,8 @@ namespace UE::NNEQA::Private::NNERuntimeRDG::UnaryOp
 	{
 		FTensor XC2x1 = MakeConstTensor(TEXT("XC2x1"), { 2,1 }, { -0.5f, 0.8f });
 		FTensor Y = MakeTensor(TEXT("Y"), { 2,1 });
-		ElementWiseUnaryCPUHelper::Apply(EElementWiseUnaryOperatorType::Relu, XC2x1, 1.0f, 2.0f, 3.0f, Y);
+		CPUHelper::ElementWiseUnary::Apply(EElementWiseUnaryOperatorType::Relu, XC2x1, 1.0f, 2.0f, 3.0f, Y);
+		UTEST_TRUE(TEXT("Y const if input is const"), Y.HasPreparedData());
 		UTEST_EQUAL(TEXT("Relu(XC2x1)[0]"), Y.GetPreparedData<float>()[0], 0.0f);
 		UTEST_EQUAL(TEXT("Relu(XC2x1)[1]"), Y.GetPreparedData<float>()[1], 0.8f);
 		return true;
@@ -281,7 +303,8 @@ namespace UE::NNEQA::Private::NNERuntimeRDG::UnaryOp
 	{
 		FTensor XC2x1 = MakeConstTensor(TEXT("XC2x1"), { 2,1 }, { 0.1f, -1.0f });
 		FTensor Y = MakeTensor(TEXT("Y"), { 2,1 });
-		ElementWiseUnaryCPUHelper::Apply(EElementWiseUnaryOperatorType::Selu, XC2x1, 1.0f, 2.0f, 3.0f, Y);
+		CPUHelper::ElementWiseUnary::Apply(EElementWiseUnaryOperatorType::Selu, XC2x1, 1.0f, 2.0f, 3.0f, Y);
+		UTEST_TRUE(TEXT("Y const if input is const"), Y.HasPreparedData());
 		UTEST_EQUAL(TEXT("Selu(XC2x1)[0]"), Y.GetPreparedData<float>()[0], 0.3f);
 		UTEST_EQUAL(TEXT("Selu(XC2x1)[1]"), Y.GetPreparedData<float>()[1], 3.0f * (0.36787944117144233f - 1.0f));
 		return true;
@@ -292,7 +315,8 @@ namespace UE::NNEQA::Private::NNERuntimeRDG::UnaryOp
 	{
 		FTensor XC2x1 = MakeConstTensor(TEXT("XC2x1"), { 2,1 }, { 0.0f, 1.0f });
 		FTensor Y = MakeTensor(TEXT("Y"), { 2,1 });
-		ElementWiseUnaryCPUHelper::Apply(EElementWiseUnaryOperatorType::Sigmoid, XC2x1, 1.0f, 2.0f, 3.0f, Y);
+		CPUHelper::ElementWiseUnary::Apply(EElementWiseUnaryOperatorType::Sigmoid, XC2x1, 1.0f, 2.0f, 3.0f, Y);
+		UTEST_TRUE(TEXT("Y const if input is const"), Y.HasPreparedData());
 		UTEST_EQUAL(TEXT("Sigmoid(XC2x1)[0]"), Y.GetPreparedData<float>()[0], 0.5f);
 		UTEST_EQUAL(TEXT("Sigmoid(XC2x1)[1]"), Y.GetPreparedData<float>()[1], 1.0f / (1.0f + 0.36787944117144233f));
 		return true;
@@ -303,7 +327,8 @@ namespace UE::NNEQA::Private::NNERuntimeRDG::UnaryOp
 	{
 		FTensor XC2x1 = MakeConstTensor(TEXT("XC2x1"), { 2,1 }, { -0.5f, 0.8f });
 		FTensor Y = MakeTensor(TEXT("Y"), { 2,1 });
-		ElementWiseUnaryCPUHelper::Apply(EElementWiseUnaryOperatorType::Sign, XC2x1, 1.0f, 2.0f, 3.0f, Y);
+		CPUHelper::ElementWiseUnary::Apply(EElementWiseUnaryOperatorType::Sign, XC2x1, 1.0f, 2.0f, 3.0f, Y);
+		UTEST_TRUE(TEXT("Y const if input is const"), Y.HasPreparedData());
 		UTEST_EQUAL(TEXT("Sign(XC2x1)[0]"), Y.GetPreparedData<float>()[0], FMath::Sign(XC2x1.GetPreparedData<float>()[0]));
 		UTEST_EQUAL(TEXT("Sign(XC2x1)[1]"), Y.GetPreparedData<float>()[1], FMath::Sign(XC2x1.GetPreparedData<float>()[1]));
 		return true;
@@ -314,7 +339,8 @@ namespace UE::NNEQA::Private::NNERuntimeRDG::UnaryOp
 	{
 		FTensor XC2x1 = MakeConstTensor(TEXT("XC2x1"), { 2,1 }, { -0.5f, 0.8f });
 		FTensor Y = MakeTensor(TEXT("Y"), { 2,1 });
-		ElementWiseUnaryCPUHelper::Apply(EElementWiseUnaryOperatorType::Sin, XC2x1, 1.0f, 2.0f, 3.0f, Y);
+		CPUHelper::ElementWiseUnary::Apply(EElementWiseUnaryOperatorType::Sin, XC2x1, 1.0f, 2.0f, 3.0f, Y);
+		UTEST_TRUE(TEXT("Y const if input is const"), Y.HasPreparedData());
 		UTEST_EQUAL(TEXT("Sin(XC2x1)[0]"), Y.GetPreparedData<float>()[0], FMath::Sin(XC2x1.GetPreparedData<float>()[0]));
 		UTEST_EQUAL(TEXT("Sin(XC2x1)[1]"), Y.GetPreparedData<float>()[1], FMath::Sin(XC2x1.GetPreparedData<float>()[1]));
 		return true;
@@ -325,7 +351,8 @@ namespace UE::NNEQA::Private::NNERuntimeRDG::UnaryOp
 	{
 		FTensor XC2x1 = MakeConstTensor(TEXT("XC2x1"), { 2,1 }, { -0.5f, 0.8f });
 		FTensor Y = MakeTensor(TEXT("Y"), { 2,1 });
-		ElementWiseUnaryCPUHelper::Apply(EElementWiseUnaryOperatorType::Sinh, XC2x1, 1.0f, 2.0f, 3.0f, Y);
+		CPUHelper::ElementWiseUnary::Apply(EElementWiseUnaryOperatorType::Sinh, XC2x1, 1.0f, 2.0f, 3.0f, Y);
+		UTEST_TRUE(TEXT("Y const if input is const"), Y.HasPreparedData());
 		UTEST_EQUAL(TEXT("Sinh(XC2x1)[0]"), Y.GetPreparedData<float>()[0], FMath::Sinh(XC2x1.GetPreparedData<float>()[0]));
 		UTEST_EQUAL(TEXT("Sinh(XC2x1)[1]"), Y.GetPreparedData<float>()[1], FMath::Sinh(XC2x1.GetPreparedData<float>()[1]));
 		return true;
@@ -336,7 +363,8 @@ namespace UE::NNEQA::Private::NNERuntimeRDG::UnaryOp
 	{
 		FTensor XC2x1 = MakeConstTensor(TEXT("XC2x1"), { 2,1 }, { 0.0f, 1.0f });
 		FTensor Y = MakeTensor(TEXT("Y"), { 2,1 });
-		ElementWiseUnaryCPUHelper::Apply(EElementWiseUnaryOperatorType::Softplus, XC2x1, 1.0f, 2.0f, 3.0f, Y);
+		CPUHelper::ElementWiseUnary::Apply(EElementWiseUnaryOperatorType::Softplus, XC2x1, 1.0f, 2.0f, 3.0f, Y);
+		UTEST_TRUE(TEXT("Y const if input is const"), Y.HasPreparedData());
 		UTEST_EQUAL(TEXT("Softplus(XC2x1)[0]"), Y.GetPreparedData<float>()[0], 0.6931471805599453f);
 		UTEST_EQUAL(TEXT("Softplus(XC2x1)[1]"), Y.GetPreparedData<float>()[1], 1.3132616875182228f);
 		return true;
@@ -347,7 +375,8 @@ namespace UE::NNEQA::Private::NNERuntimeRDG::UnaryOp
 	{
 		FTensor XC2x1 = MakeConstTensor(TEXT("XC2x1"), { 2,1 }, { 0.0f, 1.0f });
 		FTensor Y = MakeTensor(TEXT("Y"), { 2,1 });
-		ElementWiseUnaryCPUHelper::Apply(EElementWiseUnaryOperatorType::Softsign, XC2x1, 1.0f, 2.0f, 3.0f, Y);
+		CPUHelper::ElementWiseUnary::Apply(EElementWiseUnaryOperatorType::Softsign, XC2x1, 1.0f, 2.0f, 3.0f, Y);
+		UTEST_TRUE(TEXT("Y const if input is const"), Y.HasPreparedData());
 		UTEST_EQUAL(TEXT("Softsign(XC2x1)[0]"), Y.GetPreparedData<float>()[0], 0.0f);
 		UTEST_EQUAL(TEXT("Softsign(XC2x1)[1]"), Y.GetPreparedData<float>()[1], 0.5f);
 		return true;
@@ -358,7 +387,8 @@ namespace UE::NNEQA::Private::NNERuntimeRDG::UnaryOp
 	{
 		FTensor XC2x1 = MakeConstTensor(TEXT("XC2x1"), { 2,1 }, { 2.0f, 1.0f });
 		FTensor Y = MakeTensor(TEXT("Y"), { 2,1 });
-		ElementWiseUnaryCPUHelper::Apply(EElementWiseUnaryOperatorType::Sqrt, XC2x1, 1.0f, 2.0f, 3.0f, Y);
+		CPUHelper::ElementWiseUnary::Apply(EElementWiseUnaryOperatorType::Sqrt, XC2x1, 1.0f, 2.0f, 3.0f, Y);
+		UTEST_TRUE(TEXT("Y const if input is const"), Y.HasPreparedData());
 		UTEST_EQUAL(TEXT("Sqrt(XC2x1)[0]"), Y.GetPreparedData<float>()[0], FMath::Sqrt(XC2x1.GetPreparedData<float>()[0]));
 		UTEST_EQUAL(TEXT("Sqrt(XC2x1)[1]"), Y.GetPreparedData<float>()[1], FMath::Sqrt(XC2x1.GetPreparedData<float>()[1]));
 		return true;
@@ -369,7 +399,8 @@ namespace UE::NNEQA::Private::NNERuntimeRDG::UnaryOp
 	{
 		FTensor XC2x1 = MakeConstTensor(TEXT("XC2x1"), { 2,1 }, { -0.5f, 0.8f });
 		FTensor Y = MakeTensor(TEXT("Y"), { 2,1 });
-		ElementWiseUnaryCPUHelper::Apply(EElementWiseUnaryOperatorType::Tan, XC2x1, 1.0f, 2.0f, 3.0f, Y);
+		CPUHelper::ElementWiseUnary::Apply(EElementWiseUnaryOperatorType::Tan, XC2x1, 1.0f, 2.0f, 3.0f, Y);
+		UTEST_TRUE(TEXT("Y const if input is const"), Y.HasPreparedData());
 		UTEST_EQUAL(TEXT("Tan(XC2x1)[0]"), Y.GetPreparedData<float>()[0], FMath::Tan(XC2x1.GetPreparedData<float>()[0]));
 		UTEST_EQUAL(TEXT("Tan(XC2x1)[1]"), Y.GetPreparedData<float>()[1], FMath::Tan(XC2x1.GetPreparedData<float>()[1]));
 		return true;
@@ -380,7 +411,8 @@ namespace UE::NNEQA::Private::NNERuntimeRDG::UnaryOp
 	{
 		FTensor XC2x1 = MakeConstTensor(TEXT("XC2x1"), { 2,1 }, { 1.0f, 2.0f });
 		FTensor Y = MakeTensor(TEXT("Y"), { 2,1 });
-		ElementWiseUnaryCPUHelper::Apply(EElementWiseUnaryOperatorType::Tanh, XC2x1, 1.0f, 2.0f, 3.0f, Y);
+		CPUHelper::ElementWiseUnary::Apply(EElementWiseUnaryOperatorType::Tanh, XC2x1, 1.0f, 2.0f, 3.0f, Y);
+		UTEST_TRUE(TEXT("Y const if input is const"), Y.HasPreparedData());
 		UTEST_EQUAL(TEXT("Tanh(XC2x1)[0]"), Y.GetPreparedData<float>()[0], 0.7615941559557649f);
 		UTEST_EQUAL(TEXT("Tanh(XC2x1)[1]"), Y.GetPreparedData<float>()[1], 0.9640275800758169f);
 		return true;
diff --git a/Engine/Plugins/Experimental/NNE/Source/NNERuntimeRDG/Internal/NNERuntimeRDGElementWiseBinaryHelper.h b/Engine/Plugins/Experimental/NNE/Source/NNERuntimeRDG/Internal/NNERuntimeRDGElementWiseBinaryHelper.h
deleted file mode 100644
diff --git a/Engine/Plugins/Experimental/NNE/Source/NNERuntimeRDG/Internal/NNERuntimeRDGElementWiseUnaryHelper.h b/Engine/Plugins/Experimental/NNE/Source/NNERuntimeRDG/Internal/NNERuntimeRDGElementWiseUnaryHelper.h
deleted file mode 100644
diff --git a/Engine/Plugins/Experimental/NNE/Source/NNERuntimeRDG/Internal/NNERuntimeRDGHelperConcat.h b/Engine/Plugins/Experimental/NNE/Source/NNERuntimeRDG/Internal/NNERuntimeRDGHelperConcat.h
new file mode 100644
--- /dev/null
+++ b/Engine/Plugins/Experimental/NNE/Source/NNERuntimeRDG/Internal/NNERuntimeRDGHelperConcat.h
@@ -0,0 +1,11 @@
+// Copyright Epic Games, Inc. All Rights Reserved.
+
+#pragma once 
+
+#include "CoreMinimal.h"
+#include "NNECoreTensor.h"
+
+namespace UE::NNERuntimeRDG::Internal::CPUHelper::Concat
+{
+	void NNERUNTIMERDG_API Apply(TConstArrayView<NNECore::Internal::FTensorRef> InputTensors, NNECore::Internal::FTensor& OutputTensor, int32 Axis);
+} // UE::NNERuntimeRDG::Internal::CPUHelper::Concat
\ No newline at end of file
diff --git a/Engine/Plugins/Experimental/NNE/Source/NNERuntimeRDG/Internal/NNERuntimeRDGHelperElementWiseBinary.h b/Engine/Plugins/Experimental/NNE/Source/NNERuntimeRDG/Internal/NNERuntimeRDGHelperElementWiseBinary.h
new file mode 100644
--- /dev/null
+++ b/Engine/Plugins/Experimental/NNE/Source/NNERuntimeRDG/Internal/NNERuntimeRDGHelperElementWiseBinary.h
@@ -0,0 +1,14 @@
+// Copyright Epic Games, Inc. All Rights Reserved.
+
+#pragma once 
+
+#include "CoreMinimal.h"
+#include "NNECoreOperator.h"
+
+namespace UE::NNECore::Internal { class FTensor; }
+
+namespace UE::NNERuntimeRDG::Internal::CPUHelper::ElementWiseBinary
+{
+	void NNERUNTIMERDG_API Apply(NNECore::Internal::EElementWiseBinaryOperatorType OpType, const NNECore::Internal::FTensor& LHSTensor, const NNECore::Internal::FTensor& RHSTensor, NNECore::Internal::FTensor& OutputTensor);
+	
+} // UE::NNERuntimeRDG::Internal::CPUHelper::ElementWiseBinary
\ No newline at end of file
diff --git a/Engine/Plugins/Experimental/NNE/Source/NNERuntimeRDG/Internal/NNERuntimeRDGHelperElementWiseUnary.h b/Engine/Plugins/Experimental/NNE/Source/NNERuntimeRDG/Internal/NNERuntimeRDGHelperElementWiseUnary.h
new file mode 100644
--- /dev/null
+++ b/Engine/Plugins/Experimental/NNE/Source/NNERuntimeRDG/Internal/NNERuntimeRDGHelperElementWiseUnary.h
@@ -0,0 +1,13 @@
+// Copyright Epic Games, Inc. All Rights Reserved.
+
+#pragma once 
+
+#include "CoreMinimal.h"
+#include "NNECoreOperator.h"
+
+namespace UE::NNECore::Internal { class FTensor; }
+
+namespace UE::NNERuntimeRDG::Internal::CPUHelper::ElementWiseUnary
+{
+	void NNERUNTIMERDG_API Apply(NNECore::Internal::EElementWiseUnaryOperatorType OpType, const NNECore::Internal::FTensor& Tensor, float Alpha, float Beta, float Gamma, NNECore::Internal::FTensor& OutputTensor);
+} // UE::NNERuntimeRDG::Internal::CPUHelper::ElementWiseUnary
\ No newline at end of file
diff --git a/Engine/Plugins/Experimental/NNE/Source/NNERuntimeRDG/Internal/NNERuntimeRDGHelperSlice.h b/Engine/Plugins/Experimental/NNE/Source/NNERuntimeRDG/Internal/NNERuntimeRDGHelperSlice.h
new file mode 100644
--- /dev/null
+++ b/Engine/Plugins/Experimental/NNE/Source/NNERuntimeRDG/Internal/NNERuntimeRDGHelperSlice.h
@@ -0,0 +1,12 @@
+// Copyright Epic Games, Inc. All Rights Reserved.
+
+#pragma once 
+
+#include "CoreMinimal.h"
+
+namespace UE::NNECore::Internal { class FTensor; }
+
+namespace UE::NNERuntimeRDG::Internal::CPUHelper::Slice
+{
+	void NNERUNTIMERDG_API Apply(const NNECore::Internal::FTensor& InputTensor, NNECore::Internal::FTensor& OutputTensor, TConstArrayView<int32> Starts);
+} // UE::NNERuntimeRDG::Internal::CPUHelper::Slice
\ No newline at end of file
diff --git a/Engine/Plugins/Experimental/NNE/Source/NNERuntimeRDG/Private/Helper/NNERuntimeRDGHelperConcat.cpp b/Engine/Plugins/Experimental/NNE/Source/NNERuntimeRDG/Private/Helper/NNERuntimeRDGHelperConcat.cpp
new file mode 100644
--- /dev/null
+++ b/Engine/Plugins/Experimental/NNE/Source/NNERuntimeRDG/Private/Helper/NNERuntimeRDGHelperConcat.cpp
@@ -0,0 +1,62 @@
+// Copyright Epic Games, Inc. All Rights Reserved.
+
+#include "NNERuntimeRDGHelperConcat.h"
+
+#include "NNECoreTypes.h"
+#include "NNERuntimeRDGTensorIdxIterator.h"
+
+namespace UE::NNERuntimeRDG::Internal::CPUHelper::Concat
+{
+	void Apply(TConstArrayView<NNECore::Internal::FTensorRef> InputTensors, NNECore::Internal::FTensor& OutputTensor, int32 Axis)
+	{
+		static constexpr int32 MaxItemInOutputTensor = NNECore::FTensorShape::MaxRank * 2;
+		
+		if (OutputTensor.GetVolume() >= MaxItemInOutputTensor)
+		{
+			return;
+		}
+
+		if (OutputTensor.GetDataType() != ENNETensorDataType::Float)
+		{
+			return;
+		}
+		
+		for (NNECore::Internal::FTensorRef InputTensor : InputTensors)
+		{
+			check(InputTensor != nullptr);
+			if (!InputTensor->HasPreparedData())
+			{
+				return;
+			}
+		}
+
+		check(Axis >= 0 && Axis < OutputTensor.GetShape().Rank());
+
+		TArray<float> OutputData;
+		const Private::TensorIdxIterator itOutput(OutputTensor.GetShape());
+		int32 ConcatAxisOffset = 0; 
+
+		OutputData.SetNumUninitialized(OutputTensor.GetVolume());
+		
+		for (NNECore::Internal::FTensorRef InputTensor : InputTensors)
+		{
+			TConstArrayView<float> InputData = InputTensor->GetPreparedData<float>();
+			Private::TensorIdxIterator itInput(InputTensor->GetShape());
+
+			do
+			{
+				TArray<uint32> CurOutputPosition(itInput.GetPositions());
+				CurOutputPosition[Axis] += ConcatAxisOffset;
+				
+				float Value = InputData[itInput.GetIndex()];
+				OutputData[itOutput.GetIndexFromPosition(CurOutputPosition)] = Value;
+				
+			} while (itInput.Advance());
+			
+			ConcatAxisOffset += InputTensor->GetShape().GetData()[Axis];
+		}
+
+		OutputTensor.SetPreparedData<float>(OutputData);
+	}
+	
+} // UE::NNERuntimeRDG::Internal::CPUHelper::Concat
diff --git a/Engine/Plugins/Experimental/NNE/Source/NNERuntimeRDG/Private/Helper/NNERuntimeRDGHelperElementWiseBinary.cpp b/Engine/Plugins/Experimental/NNE/Source/NNERuntimeRDG/Private/Helper/NNERuntimeRDGHelperElementWiseBinary.cpp
new file mode 100644
--- /dev/null
+++ b/Engine/Plugins/Experimental/NNE/Source/NNERuntimeRDG/Private/Helper/NNERuntimeRDGHelperElementWiseBinary.cpp
@@ -0,0 +1,81 @@
+// Copyright Epic Games, Inc. All Rights Reserved.
+
+#include "NNERuntimeRDGHelperElementWiseBinary.h"
+#include "NNERuntimeRDGTensorIdxIterator.h"
+#include "NNECoreTypes.h"
+#include "NNECoreTensor.h"
+#include "Math/UnrealMathUtility.h"
+#include "MathUtil.h"
+
+namespace UE::NNERuntimeRDG::Internal::CPUHelper::ElementWiseBinary
+{
+	template<NNECore::Internal::EElementWiseBinaryOperatorType OpType> float Apply(float X, float Y);
+	template<> float Apply<NNECore::Internal::EElementWiseBinaryOperatorType::Add>(float X, float Y) { return X + Y; }
+	template<> float Apply<NNECore::Internal::EElementWiseBinaryOperatorType::Div>(float X, float Y) { return X / Y; }
+	template<> float Apply<NNECore::Internal::EElementWiseBinaryOperatorType::Mod>(float X, float Y) { return FMath::Fmod(X, Y); }
+	template<> float Apply<NNECore::Internal::EElementWiseBinaryOperatorType::Mul>(float X, float Y) { return X * Y; }
+	template<> float Apply<NNECore::Internal::EElementWiseBinaryOperatorType::Prelu>(float X, float Y) { return (X < 0.0f) ? (Y * X) : X; }
+	template<> float Apply<NNECore::Internal::EElementWiseBinaryOperatorType::Pow>(float X, float Y) { return FMath::Pow(X, Y); }
+	template<> float Apply<NNECore::Internal::EElementWiseBinaryOperatorType::Sub>(float X, float Y) { return X - Y; }
+
+	template<NNECore::Internal::EElementWiseBinaryOperatorType OpType> void Apply(const NNECore::Internal::FTensor& LHSTensor, const NNECore::Internal::FTensor& RHSTensor, NNECore::Internal::FTensor& OutputTensor)
+	{
+		//Heuristic to avoid unexpected performance hit. This helper being intended for shape related arithmetic only.
+		static constexpr int32 MaxItemInInputTensors = NNECore::FTensorShape::MaxRank * 2;
+
+		if (LHSTensor.HasPreparedData() && 
+			RHSTensor.HasPreparedData() && 
+			(LHSTensor.GetVolume() <= MaxItemInInputTensors) &&
+			(RHSTensor.GetVolume() <= MaxItemInInputTensors))
+		{
+			TConstArrayView<float> LHSData = LHSTensor.GetPreparedData<float>();
+			TConstArrayView<float> RHSData = RHSTensor.GetPreparedData<float>();
+			TArray<float> OutputData;
+			OutputData.Reserve(OutputTensor.GetVolume());
+
+			Private::TensorIdxIterator it(OutputTensor.GetShape());
+			do
+			{
+				int32 LHSIdx = it.GetIndexToBroadcastedShape(LHSTensor.GetShape());
+				int32 RHSIdx = it.GetIndexToBroadcastedShape(RHSTensor.GetShape());
+				float LHSValue = LHSData.GetData()[LHSIdx];
+				float RHSValue = RHSData.GetData()[RHSIdx];
+				OutputData.Add(Apply<OpType>(LHSValue, RHSValue));
+			} while (it.Advance());
+
+			check(OutputData.Num() == OutputTensor.GetVolume());
+			OutputTensor.SetPreparedData<float>(OutputData);
+		}
+	}
+
+	void Apply(NNECore::Internal::EElementWiseBinaryOperatorType OpType, const NNECore::Internal::FTensor& LHSTensor, const NNECore::Internal::FTensor& RHSTensor, NNECore::Internal::FTensor& OutputTensor)
+	{
+		switch (OpType)
+		{
+		case NNECore::Internal::EElementWiseBinaryOperatorType::Add:
+			Apply<NNECore::Internal::EElementWiseBinaryOperatorType::Add>(LHSTensor, RHSTensor, OutputTensor);
+			break;
+		case NNECore::Internal::EElementWiseBinaryOperatorType::Div:
+			Apply<NNECore::Internal::EElementWiseBinaryOperatorType::Div>(LHSTensor, RHSTensor, OutputTensor);
+			break;
+		case NNECore::Internal::EElementWiseBinaryOperatorType::Mod:
+			Apply<NNECore::Internal::EElementWiseBinaryOperatorType::Mod>(LHSTensor, RHSTensor, OutputTensor);
+			break;
+		case NNECore::Internal::EElementWiseBinaryOperatorType::Mul:
+			Apply<NNECore::Internal::EElementWiseBinaryOperatorType::Mul>(LHSTensor, RHSTensor, OutputTensor);
+			break;
+		case NNECore::Internal::EElementWiseBinaryOperatorType::Prelu:
+			Apply<NNECore::Internal::EElementWiseBinaryOperatorType::Prelu>(LHSTensor, RHSTensor, OutputTensor);
+			break;
+		case NNECore::Internal::EElementWiseBinaryOperatorType::Pow:
+			Apply<NNECore::Internal::EElementWiseBinaryOperatorType::Pow>(LHSTensor, RHSTensor, OutputTensor);
+			break;
+		case NNECore::Internal::EElementWiseBinaryOperatorType::Sub:
+			Apply<NNECore::Internal::EElementWiseBinaryOperatorType::Sub>(LHSTensor, RHSTensor, OutputTensor);
+			break;
+		default:
+			break;
+		}
+	}
+	
+} // UE::NNERuntimeRDG::Internal::CPUHelper::ElementWiseBinary
diff --git a/Engine/Plugins/Experimental/NNE/Source/NNERuntimeRDG/Private/Helper/NNERuntimeRDGHelperElementWiseUnary.cpp b/Engine/Plugins/Experimental/NNE/Source/NNERuntimeRDG/Private/Helper/NNERuntimeRDGHelperElementWiseUnary.cpp
new file mode 100644
--- /dev/null
+++ b/Engine/Plugins/Experimental/NNE/Source/NNERuntimeRDG/Private/Helper/NNERuntimeRDGHelperElementWiseUnary.cpp
@@ -0,0 +1,256 @@
+// Copyright Epic Games, Inc. All Rights Reserved.
+
+#include "NNERuntimeRDGHelperElementWiseUnary.h"
+#include "NNECoreTensor.h"
+#include "Math/UnrealMathUtility.h"
+
+namespace UE::NNERuntimeRDG::Internal::CPUHelper::ElementWiseUnary
+{
+	template<NNECore::Internal::EElementWiseUnaryOperatorType OpType> float Apply(float X, float Alpha, float Beta, float Gamma);
+	
+	template<> float Apply<NNECore::Internal::EElementWiseUnaryOperatorType::Abs>(float X, float Alpha, float Beta, float Gamma) { return FMath::Abs(X); }
+
+	template<> float Apply<NNECore::Internal::EElementWiseUnaryOperatorType::Acos>(float X, float Alpha, float Beta, float Gamma) { return FMath::Acos(X); }
+	
+	template<> float Apply<NNECore::Internal::EElementWiseUnaryOperatorType::Acosh>(float X, float Alpha, float Beta, float Gamma) {
+		//https://mathworld.wolfram.com/InverseHyperbolicCosine.html
+		float FloatNan = FMath::Sqrt(-1.0f);
+		float yAboveOne = FMath::Loge(X + FMath::Sqrt(X + 1.0f) * FMath::Sqrt(X - 1.0f));
+		if (X == 1.0f)
+		{
+			return 0.0f;
+		}
+		else
+		{
+			return (X >= 1.0f) ? yAboveOne : FloatNan;
+		}
+	}
+	
+	template<> float Apply<NNECore::Internal::EElementWiseUnaryOperatorType::Asin>(float X, float Alpha, float Beta, float Gamma) { return FMath::Asin(X); }
+	
+	template<> float Apply<NNECore::Internal::EElementWiseUnaryOperatorType::Asinh>(float X, float Alpha, float Beta, float Gamma) {
+		//https://mathworld.wolfram.com/InverseHyperbolicSine.html
+		return FMath::Loge(X + FMath::Sqrt(1 + (X * X)));
+	}
+	
+	template<> float Apply<NNECore::Internal::EElementWiseUnaryOperatorType::Atan>(float X, float Alpha, float Beta, float Gamma) { return FMath::Atan(X); }
+	
+	template<> float Apply<NNECore::Internal::EElementWiseUnaryOperatorType::Atanh>(float X, float Alpha, float Beta, float Gamma) {
+		//https://mathworld.wolfram.com/InverseHyperbolicTangent.html
+		return 0.5f * (FMath::Loge(1 + X) - FMath::Loge(1 - X));
+	}
+
+	template<> float Apply<NNECore::Internal::EElementWiseUnaryOperatorType::Ceil>(float X, float Alpha, float Beta, float Gamma) { return FMath::CeilToFloat(X); }
+	
+	template<> float Apply<NNECore::Internal::EElementWiseUnaryOperatorType::Cos>(float X, float Alpha, float Beta, float Gamma) { return FMath::Cos(X); }
+	
+	template<> float Apply<NNECore::Internal::EElementWiseUnaryOperatorType::Cosh>(float X, float Alpha, float Beta, float Gamma) {
+		//https://mathworld.wolfram.com/HyperbolicCosine.html
+		return 0.5f * (FMath::Exp(X) + FMath::Exp(-X));
+	}
+	
+	template<> float Apply<NNECore::Internal::EElementWiseUnaryOperatorType::Elu>(float X, float Alpha, float Beta, float Gamma) {
+		//[REDACTED_PROFILE_URL]/onnx/blob/main/docs/Operators.md#elu
+		float yNeg = Alpha * (FMath::Exp(X) - 1.0f);
+		float yPosOrZero = X;
+		return (X >= 0.0f) ? yPosOrZero : yNeg;
+	}
+	
+	template<> float Apply<NNECore::Internal::EElementWiseUnaryOperatorType::Exp>(float X, float Alpha, float Beta, float Gamma) { return FMath::Exp(X); }
+
+	template<> float Apply<NNECore::Internal::EElementWiseUnaryOperatorType::Floor>(float X, float Alpha, float Beta, float Gamma) { return FMath::Floor(X); }
+
+	template<> float Apply<NNECore::Internal::EElementWiseUnaryOperatorType::IsInf>(float X, float Alpha, float Beta, float Gamma) { return !FMath::IsFinite(X); }
+
+	template<> float Apply<NNECore::Internal::EElementWiseUnaryOperatorType::IsNan>(float X, float Alpha, float Beta, float Gamma) { return FMath::IsNaN(X); }
+
+	template<> float Apply<NNECore::Internal::EElementWiseUnaryOperatorType::HardSigmoid>(float X, float Alpha, float Beta, float Gamma) {
+		//[REDACTED_PROFILE_URL]/onnx/blob/main/docs/Operators.md#hardSigmoid
+		return FMath::Max(0.0f, FMath::Min(1.0f, Alpha * X + Beta));
+	}
+
+	template<> float Apply<NNECore::Internal::EElementWiseUnaryOperatorType::HardSwish>(float X, float Alpha, float Beta, float Gamma) {
+		//[REDACTED_PROFILE_URL]/onnx/blob/main/docs/Operators.md#hardSwish
+		return Apply<NNECore::Internal::EElementWiseUnaryOperatorType::HardSigmoid>(X, 1.0f / 6.0f, 0.5f, Gamma);
+	}
+
+	template<> float Apply<NNECore::Internal::EElementWiseUnaryOperatorType::LeakyRelu>(float X, float Alpha, float Beta, float Gamma) { return (X >= 0.0f) ? X : Alpha * X; }
+
+	template<> float Apply<NNECore::Internal::EElementWiseUnaryOperatorType::Log>(float X, float Alpha, float Beta, float Gamma) { return FMath::Loge(X); }
+
+	template<> float Apply<NNECore::Internal::EElementWiseUnaryOperatorType::Neg>(float X, float Alpha, float Beta, float Gamma) { return -X; }
+
+	template<> float Apply<NNECore::Internal::EElementWiseUnaryOperatorType::Reciprocal>(float X, float Alpha, float Beta, float Gamma) { return 1.0f/X; }
+
+	template<> float Apply<NNECore::Internal::EElementWiseUnaryOperatorType::Relu>(float X, float Alpha, float Beta, float Gamma) { return FMath::Max(X, 0.0f); }
+
+	template<> float Apply<NNECore::Internal::EElementWiseUnaryOperatorType::Round>(float X, float Alpha, float Beta, float Gamma) { return FMath::RoundToFloat(X); }
+
+	template<> float Apply<NNECore::Internal::EElementWiseUnaryOperatorType::Selu>(float X, float Alpha, float Beta, float Gamma) {
+		//[REDACTED_PROFILE_URL]/onnx/blob/main/docs/Operators.md#Selu
+		float yNegOrZero = Gamma * (Alpha * FMath::Exp(X) - Alpha);
+		float yPos = Gamma * X;
+		return (X > 0.0f) ? yPos : yNegOrZero;
+	}
+	
+	template<> float Apply<NNECore::Internal::EElementWiseUnaryOperatorType::Sigmoid>(float X, float Alpha, float Beta, float Gamma) { return 1.0f / (1.0f + FMath::Exp(-X)); }
+	
+	template<> float Apply<NNECore::Internal::EElementWiseUnaryOperatorType::Sign>(float X, float Alpha, float Beta, float Gamma) { return FMath::Sign(X); }
+	
+	template<> float Apply<NNECore::Internal::EElementWiseUnaryOperatorType::Sin>(float X, float Alpha, float Beta, float Gamma) { return FMath::Sin(X); }
+	
+	template<> float Apply<NNECore::Internal::EElementWiseUnaryOperatorType::Sinh>(float X, float Alpha, float Beta, float Gamma) { return FMath::Sinh(X); }
+	
+	template<> float Apply<NNECore::Internal::EElementWiseUnaryOperatorType::Softplus>(float X, float Alpha, float Beta, float Gamma) { return FMath::Loge(FMath::Exp(X) + 1.0f); }
+	
+	template<> float Apply<NNECore::Internal::EElementWiseUnaryOperatorType::Softsign>(float X, float Alpha, float Beta, float Gamma) { return X / (1.0f + FMath::Abs(X)); }
+	
+	template<> float Apply<NNECore::Internal::EElementWiseUnaryOperatorType::Sqrt>(float X, float Alpha, float Beta, float Gamma) { return FMath::Sqrt(X); }
+	
+	template<> float Apply<NNECore::Internal::EElementWiseUnaryOperatorType::Tan>(float X, float Alpha, float Beta, float Gamma) { return FMath::Tan(X); }
+
+	template<> float Apply<NNECore::Internal::EElementWiseUnaryOperatorType::Tanh>(float X, float Alpha, float Beta, float Gamma) {
+		//https://mathworld.wolfram.com/HyperbolicTangent.html
+		float SinhValue = Apply<NNECore::Internal::EElementWiseUnaryOperatorType::Sinh>(X, Alpha, Beta, Gamma);
+		float CoshValue = Apply<NNECore::Internal::EElementWiseUnaryOperatorType::Cosh>(X, Alpha, Beta, Gamma);
+		return SinhValue / CoshValue;
+	}
+
+	template<> float Apply<NNECore::Internal::EElementWiseUnaryOperatorType::Erf>(float X, float Alpha, float Beta, float Gamma) {
+		//https://aapt.scitation.org/doi/abs/10.1119/1.15018?journalCode=ajp
+		float a = 167.0f / 148.0f;
+		float b = 11.0f / 109.0f;
+		float x3 = X * X * X;
+		return Apply<NNECore::Internal::EElementWiseUnaryOperatorType::Tanh>(a * X + b * x3, Alpha, Beta, Gamma);
+	}
+
+	template<NNECore::Internal::EElementWiseUnaryOperatorType OpType> void Apply(const NNECore::Internal::FTensor& Tensor, float Alpha, float Beta, float Gamma, NNECore::Internal::FTensor& OutputTensor)
+	{
+		//Heuristic to avoid unexpected performance hit. This helper being intended for shape related arithmetic only.
+		static constexpr int32 MaxItemInInputTensors = NNECore::FTensorShape::MaxRank * 2;
+
+		if (Tensor.HasPreparedData() && (Tensor.GetVolume() <= MaxItemInInputTensors))
+		{
+			TConstArrayView<float> TensorData = Tensor.GetPreparedData<float>();
+			TArray<float> OutputData;
+			OutputData.Reserve(TensorData.Num());
+			for (float elem : TensorData)
+			{
+				OutputData.Add(Apply<OpType>(elem, Alpha, Beta, Gamma));
+			}
+			OutputTensor.SetPreparedData<float>(OutputData);
+		}
+	}
+
+	void Apply(NNECore::Internal::EElementWiseUnaryOperatorType OpType, const NNECore::Internal::FTensor& Tensor, float Alpha, float Beta, float Gamma, NNECore::Internal::FTensor& OutputTensor)
+	{
+		switch (OpType)
+		{
+		case NNECore::Internal::EElementWiseUnaryOperatorType::Abs:
+			Apply<NNECore::Internal::EElementWiseUnaryOperatorType::Abs>(Tensor, Alpha, Beta, Gamma, OutputTensor);
+			break;
+		case NNECore::Internal::EElementWiseUnaryOperatorType::Acos:
+			Apply<NNECore::Internal::EElementWiseUnaryOperatorType::Acos>(Tensor, Alpha, Beta, Gamma, OutputTensor);
+			break;
+		case NNECore::Internal::EElementWiseUnaryOperatorType::Acosh:
+			Apply<NNECore::Internal::EElementWiseUnaryOperatorType::Acosh>(Tensor, Alpha, Beta, Gamma, OutputTensor);
+			break;
+		case NNECore::Internal::EElementWiseUnaryOperatorType::Asin:
+			Apply<NNECore::Internal::EElementWiseUnaryOperatorType::Asin>(Tensor, Alpha, Beta, Gamma, OutputTensor);
+			break;
+		case NNECore::Internal::EElementWiseUnaryOperatorType::Asinh:
+			Apply<NNECore::Internal::EElementWiseUnaryOperatorType::Asinh>(Tensor, Alpha, Beta, Gamma, OutputTensor);
+			break;
+		case NNECore::Internal::EElementWiseUnaryOperatorType::Atan:
+			Apply<NNECore::Internal::EElementWiseUnaryOperatorType::Atan>(Tensor, Alpha, Beta, Gamma, OutputTensor);
+			break;
+		case NNECore::Internal::EElementWiseUnaryOperatorType::Atanh:
+			Apply<NNECore::Internal::EElementWiseUnaryOperatorType::Atanh>(Tensor, Alpha, Beta, Gamma, OutputTensor);
+			break;
+		case NNECore::Internal::EElementWiseUnaryOperatorType::Ceil:
+			Apply<NNECore::Internal::EElementWiseUnaryOperatorType::Ceil>(Tensor, Alpha, Beta, Gamma, OutputTensor);
+			break;
+		case NNECore::Internal::EElementWiseUnaryOperatorType::Cos:
+			Apply<NNECore::Internal::EElementWiseUnaryOperatorType::Cos>(Tensor, Alpha, Beta, Gamma, OutputTensor);
+			break;
+		case NNECore::Internal::EElementWiseUnaryOperatorType::Cosh:
+			Apply<NNECore::Internal::EElementWiseUnaryOperatorType::Cosh>(Tensor, Alpha, Beta, Gamma, OutputTensor);
+			break;
+		case NNECore::Internal::EElementWiseUnaryOperatorType::Elu:
+			Apply<NNECore::Internal::EElementWiseUnaryOperatorType::Elu>(Tensor, Alpha, Beta, Gamma, OutputTensor);
+			break;
+		case NNECore::Internal::EElementWiseUnaryOperatorType::Erf:
+			Apply<NNECore::Internal::EElementWiseUnaryOperatorType::Erf>(Tensor, Alpha, Beta, Gamma, OutputTensor);
+			break;
+		case NNECore::Internal::EElementWiseUnaryOperatorType::Exp:
+			Apply<NNECore::Internal::EElementWiseUnaryOperatorType::Exp>(Tensor, Alpha, Beta, Gamma, OutputTensor);
+			break;
+		case NNECore::Internal::EElementWiseUnaryOperatorType::Floor:
+			Apply<NNECore::Internal::EElementWiseUnaryOperatorType::Floor>(Tensor, Alpha, Beta, Gamma, OutputTensor);
+			break;
+		case NNECore::Internal::EElementWiseUnaryOperatorType::IsInf:
+			Apply<NNECore::Internal::EElementWiseUnaryOperatorType::IsInf>(Tensor, Alpha, Beta, Gamma, OutputTensor);
+			break;
+		case NNECore::Internal::EElementWiseUnaryOperatorType::IsNan:
+			Apply<NNECore::Internal::EElementWiseUnaryOperatorType::IsNan>(Tensor, Alpha, Beta, Gamma, OutputTensor);
+			break;
+		case NNECore::Internal::EElementWiseUnaryOperatorType::HardSigmoid:
+			Apply<NNECore::Internal::EElementWiseUnaryOperatorType::HardSigmoid>(Tensor, Alpha, Beta, Gamma, OutputTensor);
+			break;
+		case NNECore::Internal::EElementWiseUnaryOperatorType::HardSwish:
+			Apply<NNECore::Internal::EElementWiseUnaryOperatorType::HardSwish>(Tensor, Alpha, Beta, Gamma, OutputTensor);
+			break;
+		case NNECore::Internal::EElementWiseUnaryOperatorType::LeakyRelu:
+			Apply<NNECore::Internal::EElementWiseUnaryOperatorType::LeakyRelu>(Tensor, Alpha, Beta, Gamma, OutputTensor);
+			break;
+		case NNECore::Internal::EElementWiseUnaryOperatorType::Log:
+			Apply<NNECore::Internal::EElementWiseUnaryOperatorType::Log>(Tensor, Alpha, Beta, Gamma, OutputTensor);
+			break;
+		case NNECore::Internal::EElementWiseUnaryOperatorType::Neg:
+			Apply<NNECore::Internal::EElementWiseUnaryOperatorType::Neg>(Tensor, Alpha, Beta, Gamma, OutputTensor);
+			break;
+		case NNECore::Internal::EElementWiseUnaryOperatorType::Reciprocal:
+			Apply<NNECore::Internal::EElementWiseUnaryOperatorType::Reciprocal>(Tensor, Alpha, Beta, Gamma, OutputTensor);
+			break;
+		case NNECore::Internal::EElementWiseUnaryOperatorType::Relu:
+			Apply<NNECore::Internal::EElementWiseUnaryOperatorType::Relu>(Tensor, Alpha, Beta, Gamma, OutputTensor);
+			break;
+		case NNECore::Internal::EElementWiseUnaryOperatorType::Round:
+			Apply<NNECore::Internal::EElementWiseUnaryOperatorType::Round>(Tensor, Alpha, Beta, Gamma, OutputTensor);
+			break;
+		case NNECore::Internal::EElementWiseUnaryOperatorType::Selu:
+			Apply<NNECore::Internal::EElementWiseUnaryOperatorType::Selu>(Tensor, Alpha, Beta, Gamma, OutputTensor);
+			break;
+		case NNECore::Internal::EElementWiseUnaryOperatorType::Sigmoid:
+			Apply<NNECore::Internal::EElementWiseUnaryOperatorType::Sigmoid>(Tensor, Alpha, Beta, Gamma, OutputTensor);
+			break;
+		case NNECore::Internal::EElementWiseUnaryOperatorType::Sign:
+			Apply<NNECore::Internal::EElementWiseUnaryOperatorType::Sign>(Tensor, Alpha, Beta, Gamma, OutputTensor);
+			break;
+		case NNECore::Internal::EElementWiseUnaryOperatorType::Sin:
+			Apply<NNECore::Internal::EElementWiseUnaryOperatorType::Sin>(Tensor, Alpha, Beta, Gamma, OutputTensor);
+			break;
+		case NNECore::Internal::EElementWiseUnaryOperatorType::Sinh:
+			Apply<NNECore::Internal::EElementWiseUnaryOperatorType::Sinh>(Tensor, Alpha, Beta, Gamma, OutputTensor);
+			break;
+		case NNECore::Internal::EElementWiseUnaryOperatorType::Softplus:
+			Apply<NNECore::Internal::EElementWiseUnaryOperatorType::Softplus>(Tensor, Alpha, Beta, Gamma, OutputTensor);
+			break;
+		case NNECore::Internal::EElementWiseUnaryOperatorType::Softsign:
+			Apply<NNECore::Internal::EElementWiseUnaryOperatorType::Softsign>(Tensor, Alpha, Beta, Gamma, OutputTensor);
+			break;
+		case NNECore::Internal::EElementWiseUnaryOperatorType::Sqrt:
+			Apply<NNECore::Internal::EElementWiseUnaryOperatorType::Sqrt>(Tensor, Alpha, Beta, Gamma, OutputTensor);
+			break;
+		case NNECore::Internal::EElementWiseUnaryOperatorType::Tan:
+			Apply<NNECore::Internal::EElementWiseUnaryOperatorType::Tan>(Tensor, Alpha, Beta, Gamma, OutputTensor);
+			break;
+		case NNECore::Internal::EElementWiseUnaryOperatorType::Tanh:
+			Apply<NNECore::Internal::EElementWiseUnaryOperatorType::Tanh>(Tensor, Alpha, Beta, Gamma, OutputTensor);
+			break;
+		default:
+			break;
+		}
+	}
+	
+} // UE::NNERuntimeRDG::Internal::CPUHelper::ElementWiseUnary
diff --git a/Engine/Plugins/Experimental/NNE/Source/NNERuntimeRDG/Private/Helper/NNERuntimeRDGHelperSlice.cpp b/Engine/Plugins/Experimental/NNE/Source/NNERuntimeRDG/Private/Helper/NNERuntimeRDGHelperSlice.cpp
new file mode 100644
--- /dev/null
+++ b/Engine/Plugins/Experimental/NNE/Source/NNERuntimeRDG/Private/Helper/NNERuntimeRDGHelperSlice.cpp
@@ -0,0 +1,52 @@
+// Copyright Epic Games, Inc. All Rights Reserved.
+
+#include "NNERuntimeRDGHelperSlice.h"
+#include "NNERuntimeRDGTensorIdxIterator.h"
+
+namespace UE::NNERuntimeRDG::Internal::CPUHelper::Slice
+{
+	void Apply(const NNECore::Internal::FTensor& InputTensor, NNECore::Internal::FTensor& OutputTensor, TConstArrayView<int32> Starts)
+	{
+		static constexpr int32 MaxItemInOutputTensor = NNECore::FTensorShape::MaxRank * 2;
+
+		if (OutputTensor.GetVolume() >= MaxItemInOutputTensor)
+		{
+			return;
+		}
+
+		if (OutputTensor.GetDataType() != ENNETensorDataType::Float)
+		{
+			return;
+		}
+
+		if (!InputTensor.HasPreparedData())
+		{
+			return;
+		}
+
+		check(InputTensor.GetShape().Rank() == Starts.Num());
+		check(OutputTensor.GetShape().Rank() == Starts.Num());
+
+		TArray<float> OutputData;
+		TConstArrayView<float> InputData = InputTensor.GetPreparedData<float>();
+		Private::TensorIdxIterator itOutput(OutputTensor.GetShape());
+		const Private::TensorIdxIterator itInput(InputTensor.GetShape());
+
+		OutputData.SetNumUninitialized(OutputTensor.GetVolume());
+		do
+		{
+			TArray<uint32> CurInputPosition(itOutput.GetPositions());
+			for (int r = 0; r < CurInputPosition.Num(); ++r)
+			{
+				CurInputPosition[r] += Starts[r];
+			}
+
+			float Value = InputData[itInput.GetIndexFromPosition(CurInputPosition)];
+			OutputData[itOutput.GetIndex()] = Value;
+
+		} while (itOutput.Advance());
+
+		OutputTensor.SetPreparedData<float>(OutputData);
+	}
+	
+} // UE::NNERuntimeRDG::Private::CPUHelper::Slice
diff --git a/Engine/Plugins/Experimental/NNE/Source/NNERuntimeRDG/Private/Helper/NNERuntimeRDGTensorIdxIterator.cpp b/Engine/Plugins/Experimental/NNE/Source/NNERuntimeRDG/Private/Helper/NNERuntimeRDGTensorIdxIterator.cpp
new file mode 100644
--- /dev/null
+++ b/Engine/Plugins/Experimental/NNE/Source/NNERuntimeRDG/Private/Helper/NNERuntimeRDGTensorIdxIterator.cpp
@@ -0,0 +1,72 @@
+// Copyright Epic Games, Inc. All Rights Reserved.
+
+#include "NNERuntimeRDGTensorIdxIterator.h"
+
+namespace UE::NNERuntimeRDG::Private
+{
+	TensorIdxIterator::TensorIdxIterator(const NNECore::FTensorShape& InTensorShape) : TensorShape(InTensorShape)
+	{
+		CurrentPosition.Init(0, InTensorShape.Rank());
+	}
+
+	bool TensorIdxIterator::Advance()
+	{
+		for (int32 i = TensorShape.Rank() - 1; i >= 0; --i)
+		{
+			++CurrentPosition[i];
+			if (CurrentPosition[i] < TensorShape.GetData()[i])
+			{
+				return true;
+			}
+			CurrentPosition[i] = 0;
+		}
+		return false;
+	}
+
+	TConstArrayView<uint32> TensorIdxIterator::GetPositions() const
+	{
+		return CurrentPosition;
+	}
+
+	int32 TensorIdxIterator::GetIndexToBroadcastedShape(const NNECore::FTensorShape& InTensorShape) const
+	{
+		int32 Index = 0;
+		int32 DimBaseOffset = 1;
+		for (int32 r = TensorShape.Rank() - 1; r >= 0; --r)
+		{
+			if (r >= InTensorShape.Rank())
+			{
+				break;
+			}
+			Index += FMath::Min(CurrentPosition[r], InTensorShape.GetData()[r]-1) * DimBaseOffset;
+			DimBaseOffset *= FMath::Min(TensorShape.GetData()[r], InTensorShape.GetData()[r]);
+		}
+		return Index;
+	}
+
+	int32 TensorIdxIterator::GetIndex() const
+	{
+		int32 Index = 0;
+		int32 DimBaseOffset = 1;
+		for (int32 r = TensorShape.Rank() - 1; r >= 0; --r)
+		{
+			Index += CurrentPosition[r] * DimBaseOffset;
+			DimBaseOffset *= TensorShape.GetData()[r];
+		}
+		return Index;
+	}
+
+	int32 TensorIdxIterator::GetIndexFromPosition(TConstArrayView<uint32> Position) const
+	{
+		int32 Index = 0;
+		int32 DimBaseOffset = 1;
+		for (int32 r = TensorShape.Rank() - 1; r >= 0; --r)
+		{
+			check(Position[r] < TensorShape.GetData()[r]);
+			Index += Position[r] * DimBaseOffset;
+			DimBaseOffset *= TensorShape.GetData()[r];
+		}
+		return Index;
+	}
+
+} // UE::NNERuntimeRDG::Private
diff --git a/Engine/Plugins/Experimental/NNE/Source/NNERuntimeRDG/Private/Helper/NNERuntimeRDGTensorIdxIterator.h b/Engine/Plugins/Experimental/NNE/Source/NNERuntimeRDG/Private/Helper/NNERuntimeRDGTensorIdxIterator.h
new file mode 100644
--- /dev/null
+++ b/Engine/Plugins/Experimental/NNE/Source/NNERuntimeRDG/Private/Helper/NNERuntimeRDGTensorIdxIterator.h
@@ -0,0 +1,24 @@
+// Copyright Epic Games, Inc. All Rights Reserved.
+
+#pragma once
+
+#include "CoreMinimal.h"
+#include "NNECoreTensor.h"
+
+namespace UE::NNERuntimeRDG::Private
+{
+	class TensorIdxIterator
+	{
+		const NNECore::FTensorShape& TensorShape;
+		TArray<uint32, TInlineAllocator<NNECore::FTensorShape::MaxRank>> CurrentPosition;
+
+	public:
+		TensorIdxIterator(const NNECore::FTensorShape& InTensorShape);
+		bool Advance();
+		int32 GetIndexToBroadcastedShape(const NNECore::FTensorShape& InTensorShape) const;
+		int32 GetIndex() const;
+		TConstArrayView<uint32> GetPositions() const;
+		int32 GetIndexFromPosition(TConstArrayView<uint32> Position) const;
+	};
+
+} // UE::NNERuntimeRDG::Private
diff --git a/Engine/Plugins/Experimental/NNE/Source/NNERuntimeRDG/Private/Hlsl/NNERuntimeRDGConcat.cpp b/Engine/Plugins/Experimental/NNE/Source/NNERuntimeRDG/Private/Hlsl/NNERuntimeRDGConcat.cpp
new file mode 100644
--- /dev/null
+++ b/Engine/Plugins/Experimental/NNE/Source/NNERuntimeRDG/Private/Hlsl/NNERuntimeRDGConcat.cpp
@@ -0,0 +1,132 @@
+// Copyright Epic Games, Inc. All Rights Reserved.
+
+#include "NNERuntimeRDGConcat.h"
+#include "NNECoreTensor.h"
+#include "NNECoreTypes.h"
+#include "NNERuntimeRDGHelperConcat.h"
+
+namespace UE::NNERuntimeRDG::Private::Hlsl
+{
+	/**
+	 * Concat operator implementation
+	 */
+	class FConcat : public FOperatorHlsl
+	{
+	public:
+
+		FConcat() {}
+		virtual ~FConcat() = default;
+
+		int32 Axis;
+
+	public:
+
+		virtual int PrepareOutputs(TConstArrayView<NNECore::Internal::FTensorRef> InputTensors, TArrayView<NNECore::Internal::FTensorRef> OutputTensors) const override
+		{
+			check(InputTensors.Num() >= 1);
+			check(OutputTensors.Num() == 1);
+
+			TArray<uint32> OutputShapeData(InputTensors[0]->GetShape().GetData());
+			int32 AxisIndex = Axis >= 0 ? Axis : InputTensors[0]->GetShape().Rank() - Axis;
+			
+			for (int32 i = 1; i < InputTensors.Num(); ++i)
+			{
+				OutputShapeData[AxisIndex] += InputTensors[i]->GetShape().GetData()[AxisIndex];
+				
+				for (int32 r = 0; r < OutputShapeData.Num(); ++r)
+				{
+					if (r != AxisIndex && (OutputShapeData[r] != InputTensors[i]->GetShape().GetData()[r]))
+					{
+						UE_LOG(LogNNE, Warning, TEXT("Concat: all input tensors should have the same shape except on the concatenation axis"));
+						return false;
+					}
+				}
+			}
+
+			NNECore::FTensorShape OutputShape = NNECore::FTensorShape::Make(OutputShapeData);
+			
+			OutputTensors[0]->SetShape(OutputShape);
+
+			Internal::CPUHelper::Concat::Apply(InputTensors, *OutputTensors[0], Axis);
+			
+			if (!OutputTensors[0]->HasPreparedData())
+			{
+				UE_LOG(LogNNE, Warning, TEXT("Concat: Output could not be computed as a constant tensor, however Concat is not implemented on GPU at the moment."));
+				return -1;
+			}
+			
+			return 0;
+		};
+
+		virtual bool Initialize(TConstArrayView<NNECore::FTensorDesc> InputTensorDescs, TConstArrayView<NNECore::FTensorDesc> OutputTensorDescs, const NNECore::FAttributeMap& Attributes) override
+		{
+			check(InputTensorDescs.Num() >= 1);
+			check(OutputTensorDescs.Num() == 1);
+
+			Axis = Attributes.GetValue<int32>(TEXT("axis"));
+			
+			int32 InputsRank = InputTensorDescs[0].GetShape().Rank();
+			
+			for (int32 i = 1; i < InputTensorDescs.Num(); ++i)
+			{
+				if (InputsRank != InputTensorDescs[i].GetShape().Rank())
+				{
+					UE_LOG(LogNNE, Warning, TEXT("Concat: all input tensors should have the same rank"));
+					return false;
+				}
+			}
+			
+			if (Axis < -InputsRank || Axis >(InputsRank - 1))
+			{
+				UE_LOG(LogNNE, Warning, TEXT("Axis should be in range [-r,r-1] however it is %d while inputs have rank %d."), Axis, InputsRank);
+				return false;
+			}
+			
+			return true;
+		}
+
+		virtual void Dispatch(FRDGBuilder& GraphBuilder, TConstArrayView<FTensorRDGRef> InputTensors, TConstArrayView<FTensorRDGRef> OutputTensors) override
+		{
+			//Need to upload the data to gpu tensor or have the hlsl model understand that this operator result will only be needed on CPU.
+			//Current behavior will be faulty if the model rely on this operator result on GPU.
+		}
+	};
+
+	bool ValidateConcatOperator(const NNECore::FAttributeMap& AttributeMap, TConstArrayView<ENNETensorDataType> InputTypes, TConstArrayView<NNECore::FSymbolicTensorShape> InputConcats)
+	{
+		bool bIsValid = true;
+
+		//This match version 11 of the Concat operator
+		//[REDACTED_PROFILE_URL]/onnx/blob/main/docs/Operators.md#Concat
+		FAttributeValidator AttributeValidator;
+		AttributeValidator.AddRequired(TEXT("axis"), ENNEAttributeDataType::Int32);
+		bIsValid &= AttributeValidator.Validate(AttributeMap);
+
+		if (InputTypes.Num() == 0)
+		{
+			UE_LOG(LogNNE, Error, TEXT("Concat operator requires at least 1 input"));
+			bIsValid = false;
+		}
+		for (int32 i = 0; i < InputTypes.Num(); ++i)
+		{
+			if (InputTypes[i] != ENNETensorDataType::Float)
+			{
+				UE_LOG(LogNNE, Warning, TEXT("Concat operator input '%d' of type '%d' is not supported, should be float at the moment."), i, InputTypes[i]);
+				bIsValid = false;
+			}
+		}
+
+		return bIsValid;
+	}
+
+	FOperatorHlsl* CreateConcatOperator()
+	{
+		return new FConcat();
+	}
+
+	bool RegisterConcatOperator(FOperatorRegistryHlsl& Registry)
+	{
+		Registry.OpAdd(TEXT("Concat"), CreateConcatOperator, ValidateConcatOperator);
+		return true;
+	}
+} // UE::NNERuntimeRDG::Private::Hlsl
diff --git a/Engine/Plugins/Experimental/NNE/Source/NNERuntimeRDG/Private/Hlsl/NNERuntimeRDGConcat.h b/Engine/Plugins/Experimental/NNE/Source/NNERuntimeRDG/Private/Hlsl/NNERuntimeRDGConcat.h
new file mode 100644
--- /dev/null
+++ b/Engine/Plugins/Experimental/NNE/Source/NNERuntimeRDG/Private/Hlsl/NNERuntimeRDGConcat.h
@@ -0,0 +1,11 @@
+// Copyright Epic Games, Inc. All Rights Reserved.
+
+#pragma once
+
+#include "NNERuntimeRDG.h"
+#include "NNERuntimeRDGHlslOp.h"
+
+namespace UE::NNERuntimeRDG::Private::Hlsl
+{
+	bool RegisterConcatOperator(FOperatorRegistryHlsl& Registry);
+} // UE::NNERuntimeRDG::Private::Hlsl
diff --git a/Engine/Plugins/Experimental/NNE/Source/NNERuntimeRDG/Private/Hlsl/NNERuntimeRDGElementWiseBinary.cpp b/Engine/Plugins/Experimental/NNE/Source/NNERuntimeRDG/Private/Hlsl/NNERuntimeRDGElementWiseBinary.cpp
--- a/Engine/Plugins/Experimental/NNE/Source/NNERuntimeRDG/Private/Hlsl/NNERuntimeRDGElementWiseBinary.cpp
+++ b/Engine/Plugins/Experimental/NNE/Source/NNERuntimeRDG/Private/Hlsl/NNERuntimeRDGElementWiseBinary.cpp
@@ -1,7 +1,7 @@
 // Copyright Epic Games, Inc. All Rights Reserved.
 
 #include "NNERuntimeRDGElementWiseBinary.h"
-#include "NNERuntimeRDGElementWiseBinaryHelper.h"
+#include "NNERuntimeRDGHelperElementWiseBinary.h"
 #include "NNEHlslShadersElementWiseBinaryCS.h"
 #include "NNERuntimeRDGHlslHelper.h"
 #include "NNECoreTensor.h"
@@ -58,7 +58,7 @@ namespace UE::NNERuntimeRDG::Private::Hlsl
 			
 			OutputTensors[0]->SetShape(OutputShape);
 
-			Internal::ElementWiseBinaryCPUHelper::Apply(OpType, *InputTensors[0], *InputTensors[1], *OutputTensors[0]);
+			Internal::CPUHelper::ElementWiseBinary::Apply(OpType, *InputTensors[0], *InputTensors[1], *OutputTensors[0]);
 			
 			return 0;
 		}
diff --git a/Engine/Plugins/Experimental/NNE/Source/NNERuntimeRDG/Private/Hlsl/NNERuntimeRDGElementWiseUnary.cpp b/Engine/Plugins/Experimental/NNE/Source/NNERuntimeRDG/Private/Hlsl/NNERuntimeRDGElementWiseUnary.cpp
--- a/Engine/Plugins/Experimental/NNE/Source/NNERuntimeRDG/Private/Hlsl/NNERuntimeRDGElementWiseUnary.cpp
+++ b/Engine/Plugins/Experimental/NNE/Source/NNERuntimeRDG/Private/Hlsl/NNERuntimeRDGElementWiseUnary.cpp
@@ -1,7 +1,7 @@
 // Copyright Epic Games, Inc. All Rights Reserved.
 
 #include "NNERuntimeRDGElementWiseUnary.h"
-#include "NNERuntimeRDGElementWiseUnaryHelper.h"
+#include "NNERuntimeRDGHelperElementWiseUnary.h"
 #include "NNEHlslShadersElementWiseUnaryCS.h"
 #include "NNERuntimeRDGHlslHelper.h"
 #include "NNECoreAttributeMap.h"
@@ -42,7 +42,7 @@ namespace UE::NNERuntimeRDG::Private::Hlsl
 
 			const NNECore::Internal::FTensor& X = *InputTensors[0];
 
-			Internal::ElementWiseUnaryCPUHelper::Apply(OpType, X, Alpha, Beta, Gamma, *OutputTensors[0]);
+			Internal::CPUHelper::ElementWiseUnary::Apply(OpType, X, Alpha, Beta, Gamma, *OutputTensors[0]);
 			
 			return 0;
 		}
diff --git a/Engine/Plugins/Experimental/NNE/Source/NNERuntimeRDG/Private/Hlsl/NNERuntimeRDGElementWiseVariadic.cpp b/Engine/Plugins/Experimental/NNE/Source/NNERuntimeRDG/Private/Hlsl/NNERuntimeRDGElementWiseVariadic.cpp
--- a/Engine/Plugins/Experimental/NNE/Source/NNERuntimeRDG/Private/Hlsl/NNERuntimeRDGElementWiseVariadic.cpp
+++ b/Engine/Plugins/Experimental/NNE/Source/NNERuntimeRDG/Private/Hlsl/NNERuntimeRDGElementWiseVariadic.cpp
@@ -200,12 +200,6 @@ namespace UE::NNERuntimeRDG::Private::Hlsl
 			}
 		}
 		
-		FInputValidator InputValidator;
-		InputValidator.AddSupportedType(ENNETensorDataType::Float);
-		InputValidator.AddRequired();
-		InputValidator.AddRequired();
-		bIsValid &= InputValidator.Validate(InputTypes);
-
 		return bIsValid;
 	}
 
diff --git a/Engine/Plugins/Experimental/NNE/Source/NNERuntimeRDG/Private/Hlsl/NNERuntimeRDGShape.cpp b/Engine/Plugins/Experimental/NNE/Source/NNERuntimeRDG/Private/Hlsl/NNERuntimeRDGShape.cpp
--- a/Engine/Plugins/Experimental/NNE/Source/NNERuntimeRDG/Private/Hlsl/NNERuntimeRDGShape.cpp
+++ b/Engine/Plugins/Experimental/NNE/Source/NNERuntimeRDG/Private/Hlsl/NNERuntimeRDGShape.cpp
@@ -98,7 +98,7 @@ namespace UE::NNERuntimeRDG::Private::Hlsl
 
 	bool RegisterShapeOperator(FOperatorRegistryHlsl& Registry)
 	{
-		Registry.OpAdd(TEXT("Shaoe"), CreateShapeOperator, ValidateShapeOperator);
+		Registry.OpAdd(TEXT("Shape"), CreateShapeOperator, ValidateShapeOperator);
 		return true;
 	}
 } // UE::NNERuntimeRDG::Private::Hlsl
diff --git a/Engine/Plugins/Experimental/NNE/Source/NNERuntimeRDG/Private/Hlsl/NNERuntimeRDGSlice.cpp b/Engine/Plugins/Experimental/NNE/Source/NNERuntimeRDG/Private/Hlsl/NNERuntimeRDGSlice.cpp
new file mode 100644
--- /dev/null
+++ b/Engine/Plugins/Experimental/NNE/Source/NNERuntimeRDG/Private/Hlsl/NNERuntimeRDGSlice.cpp
@@ -0,0 +1,156 @@
+// Copyright Epic Games, Inc. All Rights Reserved.
+
+#include "NNERuntimeRDGSlice.h"
+#include "NNERuntimeRDGHelperSlice.h"
+#include "NNECoreTensor.h"
+#include "NNECoreTypes.h"
+
+namespace UE::NNERuntimeRDG::Private::Hlsl
+{
+	/**
+	 * Slice operator implementation
+	 */
+	class FSlice : public FOperatorHlsl
+	{
+	public:
+
+		FSlice() {}
+		virtual ~FSlice() = default;
+
+		TArray<int32> AxesAttr;
+		TArray<int32> EndsAttr;
+		TArray<int32> StartsAttr;
+
+	public:
+
+		virtual int PrepareOutputs(TConstArrayView<NNECore::Internal::FTensorRef> InputTensors, TArrayView<NNECore::Internal::FTensorRef> OutputTensors) const override
+		{
+			check(InputTensors.Num() == 1);
+			check(OutputTensors.Num() == 1);
+			
+			TConstArrayView<uint32> Dims(InputTensors[0]->GetShape().GetData());
+			int32 InputRank = Dims.Num();
+
+			TArray<int32> Axes = AxesAttr;
+			TArray<int32> Ends = EndsAttr;
+			TArray<int32> Starts = StartsAttr;
+			TArray<int32> Start;
+			TArray<int32> End;
+
+			//see [REDACTED_PROFILE_URL]/onnx/blob/main/docs/Operators.md#slice for algorithm
+			Start.SetNum(InputRank);
+			End.SetNum(InputRank);
+			for (int32 i = 0; i < InputRank; ++i)
+			{
+				Start[i] = 0;
+				End[i] = Dims[i];
+			}
+			for (int32 i = 0; i < Axes.Num(); ++i)
+			{
+				if (Axes[i] < 0)
+				{
+					Axes[i] += InputRank;
+				}
+			}
+			for (int32 i = 0; i < Starts.Num(); ++i)
+			{
+				if (Starts[i] < 0)
+				{
+					Starts[i] += Dims[Axes[i]];
+				}
+			}
+			for (int32 i = 0; i < Ends.Num(); ++i)
+			{
+				if (Ends[i] < 0)
+				{
+					Ends[i] += Dims[Axes[i]];
+				}
+			}
+			for (int32 i = 0; i < Axes.Num(); ++i)
+			{
+				Start[Axes[i]] = FMath::Clamp(Starts[i], 0, Dims[Axes[i]]);
+				End[Axes[i]] = FMath::Clamp(Ends[i], 0, Dims[Axes[i]]);
+			}
+
+			TArray<uint32> OutputShapeData;
+			
+			OutputShapeData.Reserve(InputRank);
+			for (int32 i = 0; i < InputRank; ++i)
+			{
+				OutputShapeData.Add(End[i] - Start[i]);
+			}
+			
+			NNECore::FTensorShape OutputShape = NNECore::FTensorShape::Make(OutputShapeData);
+
+			OutputTensors[0]->SetShape(OutputShape);
+			
+			Internal::CPUHelper::Slice::Apply(*InputTensors[0], *OutputTensors[0], Start);
+
+			return -1;
+		};
+
+		virtual bool Initialize(TConstArrayView<NNECore::FTensorDesc> InputTensorDescs, TConstArrayView<NNECore::FTensorDesc> OutputTensorDescs, const NNECore::FAttributeMap& Attributes) override
+		{
+			check(InputTensorDescs.Num() == 1);
+			check(OutputTensorDescs.Num() == 1);
+
+			EndsAttr = Attributes.GetValue<TArray<int32>>(TEXT("ends"));
+			StartsAttr = Attributes.GetValue<TArray<int32>>(TEXT("starts"));
+			
+			TArray<int32> AxesDefault;
+			
+			for (int i = 0; i < StartsAttr.Num(); ++i)
+			{
+				AxesDefault.Add(i);
+			}
+			
+			AxesAttr = Attributes.GetValueOrDefault<TArray<int32>>(TEXT("axes"), AxesDefault);
+
+			if (EndsAttr.Num() != StartsAttr.Num() || AxesAttr.Num() != StartsAttr.Num())
+			{
+				UE_LOG(LogNNE, Warning, TEXT("Unsqueeze: Starts, Ends and Axes must be of the same size."));
+				return false;
+			}
+			
+			return true;
+		}
+
+		virtual void Dispatch(FRDGBuilder& GraphBuilder, TConstArrayView<FTensorRDGRef> InputTensors, TConstArrayView<FTensorRDGRef> OutputTensors) override
+		{
+			//Need to upload the data to gpu tensor or have the hlsl model understand that this operator result will only be needed on CPU.
+			//Current behavior will be faulty if the model rely on this operator result on GPU.
+		}
+	};
+
+	bool ValidateSliceOperator(const NNECore::FAttributeMap& AttributeMap, TConstArrayView<ENNETensorDataType> InputTypes, TConstArrayView<NNECore::FSymbolicTensorShape> InputSlices)
+	{
+		bool bIsValid = true;
+
+		//This match version 1 of the Slice operator, next version are 10, 11 and 13
+		//[REDACTED_PROFILE_URL]/onnx/blob/main/docs/Changelog.md#Slice-1
+		FAttributeValidator AttributeValidator;
+		AttributeValidator.AddOptional(TEXT("axes"), ENNEAttributeDataType::Int32Array);
+		AttributeValidator.AddRequired(TEXT("ends"), ENNEAttributeDataType::Int32Array);
+		AttributeValidator.AddRequired(TEXT("starts"), ENNEAttributeDataType::Int32Array);
+		bIsValid &= AttributeValidator.Validate(AttributeMap);
+
+		FInputValidator InputValidator;
+		InputValidator.AddSupportedType(ENNETensorDataType::Float);
+		InputValidator.AddRequired();
+
+		bIsValid &= InputValidator.Validate(InputTypes);
+
+		return bIsValid;
+	}
+
+	FOperatorHlsl* CreateSliceOperator()
+	{
+		return new FSlice();
+	}
+
+	bool RegisterSliceOperator(FOperatorRegistryHlsl& Registry)
+	{
+		Registry.OpAdd(TEXT("Slice"), CreateSliceOperator, ValidateSliceOperator);
+		return true;
+	}
+} // UE::NNERuntimeRDG::Private::Hlsl
diff --git a/Engine/Plugins/Experimental/NNE/Source/NNERuntimeRDG/Private/Hlsl/NNERuntimeRDGSlice.h b/Engine/Plugins/Experimental/NNE/Source/NNERuntimeRDG/Private/Hlsl/NNERuntimeRDGSlice.h
new file mode 100644
--- /dev/null
+++ b/Engine/Plugins/Experimental/NNE/Source/NNERuntimeRDG/Private/Hlsl/NNERuntimeRDGSlice.h
@@ -0,0 +1,11 @@
+// Copyright Epic Games, Inc. All Rights Reserved.
+
+#pragma once
+
+#include "NNERuntimeRDG.h"
+#include "NNERuntimeRDGHlslOp.h"
+
+namespace UE::NNERuntimeRDG::Private::Hlsl
+{
+	bool RegisterSliceOperator(FOperatorRegistryHlsl& Registry);
+} // UE::NNERuntimeRDG::Private::Hlsl
diff --git a/Engine/Plugins/Experimental/NNE/Source/NNERuntimeRDG/Private/Hlsl/NNERuntimeRDGUnsqueeze.cpp b/Engine/Plugins/Experimental/NNE/Source/NNERuntimeRDG/Private/Hlsl/NNERuntimeRDGUnsqueeze.cpp
new file mode 100644
--- /dev/null
+++ b/Engine/Plugins/Experimental/NNE/Source/NNERuntimeRDG/Private/Hlsl/NNERuntimeRDGUnsqueeze.cpp
@@ -0,0 +1,119 @@
+// Copyright Epic Games, Inc. All Rights Reserved.
+
+#include "NNERuntimeRDGUnsqueeze.h"
+#include "NNECoreTensor.h"
+#include "NNECoreTypes.h"
+
+namespace UE::NNERuntimeRDG::Private::Hlsl
+{
+	/**
+	 * Unsqueeze operator implementation
+	 */
+	class FUnsqueeze : public FOperatorHlsl
+	{
+	public:
+
+		FUnsqueeze() {}
+		virtual ~FUnsqueeze() = default;
+
+		TArray<int32> Axes;
+
+	public:
+
+		virtual int PrepareOutputs(TConstArrayView<NNECore::Internal::FTensorRef> InputTensors, TArrayView<NNECore::Internal::FTensorRef> OutputTensors) const override
+		{
+			check(InputTensors.Num() == 1);
+			check(OutputTensors.Num() == 1);
+
+			const NNECore::Internal::FTensor& X = *InputTensors[0];
+			TArray<uint32> OutputShapeData(X.GetShape().GetData());
+
+			if (!X.HasPreparedData())
+			{
+				UE_LOG(LogNNE, Warning, TEXT("Unsqueeze input 'Data' (name: %s) should be constant, Unsqueeze is not implemented on GPU at the moment."), *X.GetName());
+				return -1;
+			}
+			
+			for (int32 Axe : Axes)
+			{
+				OutputShapeData.Insert(1, Axe);
+			}
+			
+			NNECore::FTensorShape OutputShape = NNECore::FTensorShape::Make(OutputShapeData);
+			
+			OutputTensors[0]->SetShape(OutputShape);
+			OutputTensors[0]->SetPreparedData<uint8>(X.GetPreparedData<uint8>());
+			
+			return 0;
+		};
+
+		virtual bool Initialize(TConstArrayView<NNECore::FTensorDesc> InputTensorDescs, TConstArrayView<NNECore::FTensorDesc> OutputTensorDescs, const NNECore::FAttributeMap& Attributes) override
+		{
+			check(InputTensorDescs.Num() == 1);
+			check(OutputTensorDescs.Num() == 1);
+
+			Axes = Attributes.GetValue<TArray<int32>>(TEXT("pads"));
+			for (int32 Axe : Axes)
+			{
+				if (Axe < 0)
+				{
+					UE_LOG(LogNNE, Warning, TEXT("Unsqueeze operator does not support negative axes"));
+					return false;
+				}
+				if (Axe >= InputTensorDescs[0].GetShape().Rank() + Axes.Num())
+				{
+					UE_LOG(LogNNE, Warning, TEXT("Unsqueeze operator does not support axes greater than the number of dimensions of the resulting tensor shape"));
+					return false;
+				}
+			}
+			
+			return true;
+		}
+
+		virtual void Dispatch(FRDGBuilder& GraphBuilder, TConstArrayView<FTensorRDGRef> InputTensors, TConstArrayView<FTensorRDGRef> OutputTensors) override
+		{
+			//Need to upload the data to gpu tensor or have the hlsl model understand that this operator result will only be needed on CPU.
+			//Current behavior will be faulty if the model rely on this operator result on GPU.
+		}
+	};
+
+	bool ValidateUnsqueezeOperator(const NNECore::FAttributeMap& AttributeMap, TConstArrayView<ENNETensorDataType> InputTypes, TConstArrayView<NNECore::FSymbolicTensorShape> InputUnsqueezes)
+	{
+		bool bIsValid = true;
+
+		//This match version 1 of the Unsqueeze operator, next version are 11 and 13
+		//[REDACTED_PROFILE_URL]/onnx/blob/main/docs/Operators.md#Unsqueeze
+		FAttributeValidator AttributeValidator;
+		AttributeValidator.AddRequired(TEXT("axes"), ENNEAttributeDataType::Int32Array);
+		bIsValid &= AttributeValidator.Validate(AttributeMap);
+
+		FInputValidator InputValidator;
+		InputValidator.AddSupportedType(ENNETensorDataType::Half);
+		InputValidator.AddSupportedType(ENNETensorDataType::Float);
+		InputValidator.AddSupportedType(ENNETensorDataType::Double);
+		InputValidator.AddSupportedType(ENNETensorDataType::Int8);
+		InputValidator.AddSupportedType(ENNETensorDataType::Int16);
+		InputValidator.AddSupportedType(ENNETensorDataType::Int32);
+		InputValidator.AddSupportedType(ENNETensorDataType::Int64);
+		InputValidator.AddSupportedType(ENNETensorDataType::UInt8);
+		InputValidator.AddSupportedType(ENNETensorDataType::UInt16);
+		InputValidator.AddSupportedType(ENNETensorDataType::UInt32);
+		InputValidator.AddSupportedType(ENNETensorDataType::UInt64);
+		InputValidator.AddRequired();
+
+		bIsValid &= InputValidator.Validate(InputTypes);
+
+		return bIsValid;
+	}
+
+	FOperatorHlsl* CreateUnsqueezeOperator()
+	{
+		return new FUnsqueeze();
+	}
+
+	bool RegisterUnsqueezeOperator(FOperatorRegistryHlsl& Registry)
+	{
+		Registry.OpAdd(TEXT("Unsqueeze"), CreateUnsqueezeOperator, ValidateUnsqueezeOperator);
+		return true;
+	}
+} // UE::NNERuntimeRDG::Private::Hlsl
diff --git a/Engine/Plugins/Experimental/NNE/Source/NNERuntimeRDG/Private/Hlsl/NNERuntimeRDGUnsqueeze.h b/Engine/Plugins/Experimental/NNE/Source/NNERuntimeRDG/Private/Hlsl/NNERuntimeRDGUnsqueeze.h
new file mode 100644
--- /dev/null
+++ b/Engine/Plugins/Experimental/NNE/Source/NNERuntimeRDG/Private/Hlsl/NNERuntimeRDGUnsqueeze.h
@@ -0,0 +1,11 @@
+// Copyright Epic Games, Inc. All Rights Reserved.
+
+#pragma once
+
+#include "NNERuntimeRDG.h"
+#include "NNERuntimeRDGHlslOp.h"
+
+namespace UE::NNERuntimeRDG::Private::Hlsl
+{
+	bool RegisterUnsqueezeOperator(FOperatorRegistryHlsl& Registry);
+} // UE::NNERuntimeRDG::Private::Hlsl
diff --git a/Engine/Plugins/Experimental/NNE/Source/NNERuntimeRDG/Private/Hlsl/NNERuntimeRDGUpsample.cpp b/Engine/Plugins/Experimental/NNE/Source/NNERuntimeRDG/Private/Hlsl/NNERuntimeRDGUpsample.cpp
--- a/Engine/Plugins/Experimental/NNE/Source/NNERuntimeRDG/Private/Hlsl/NNERuntimeRDGUpsample.cpp
+++ b/Engine/Plugins/Experimental/NNE/Source/NNERuntimeRDG/Private/Hlsl/NNERuntimeRDGUpsample.cpp
@@ -32,7 +32,7 @@ namespace UE::NNERuntimeRDG::Private::Hlsl
 			
 			if (!Scales.HasPreparedData())
 			{
-				UE_LOG(LogNNE, Warning, TEXT("Upsample input 'Scale' (name: %s) should be constant for shape inference to succeed, however it is not."), *Scales.GetName());
+				UE_LOG(LogNNE, Warning, TEXT("Upsample input 'Scale' (name: %s) should be constant for shape inference to succeed, however it is not constant."), *Scales.GetName());
 				return -1;
 			}
 
diff --git a/Engine/Plugins/Experimental/NNE/Source/NNERuntimeRDG/Private/NNERuntimeRDGElementWiseBinaryHelper.cpp b/Engine/Plugins/Experimental/NNE/Source/NNERuntimeRDG/Private/NNERuntimeRDGElementWiseBinaryHelper.cpp
deleted file mode 100644
diff --git a/Engine/Plugins/Experimental/NNE/Source/NNERuntimeRDG/Private/NNERuntimeRDGElementWiseUnaryHelper.cpp b/Engine/Plugins/Experimental/NNE/Source/NNERuntimeRDG/Private/NNERuntimeRDGElementWiseUnaryHelper.cpp
deleted file mode 100644
diff --git a/Engine/Plugins/Experimental/NNE/Source/NNERuntimeRDG/Private/NNERuntimeRDGHlsl.cpp b/Engine/Plugins/Experimental/NNE/Source/NNERuntimeRDG/Private/NNERuntimeRDGHlsl.cpp
--- a/Engine/Plugins/Experimental/NNE/Source/NNERuntimeRDG/Private/NNERuntimeRDGHlsl.cpp
+++ b/Engine/Plugins/Experimental/NNE/Source/NNERuntimeRDG/Private/NNERuntimeRDGHlsl.cpp
@@ -10,6 +10,7 @@
 #include "NNEUtilsModelOptimizer.h"
 #include "Hlsl/NNERuntimeRDGConv.h"
 #include "Hlsl/NNERuntimeRDGConv.h"
+#include "Hlsl/NNERuntimeRDGConcat.h"
 #include "Hlsl/NNERuntimeRDGConvTranspose.h"
 #include "Hlsl/NNERuntimeRDGElementWiseBinary.h"
 #include "Hlsl/NNERuntimeRDGElementWiseUnary.h"
@@ -18,6 +19,8 @@
 #include "Hlsl/NNERuntimeRDGInstanceNormalization.h"
 #include "Hlsl/NNERuntimeRDGPad.h"
 #include "Hlsl/NNERuntimeRDGShape.h"
+#include "Hlsl/NNERuntimeRDGSlice.h"
+#include "Hlsl/NNERuntimeRDGUnsqueeze.h"
 #include "Hlsl/NNERuntimeRDGUpsample.h"
 #include "Hlsl/NNERuntimeRDGMatMul.h"
 
@@ -42,6 +45,9 @@ bool UNNERuntimeRDGHlslImpl::Init()
 	RegisterUpsampleOperator(*registry);
 	RegisterPadOperator(*registry);
 	RegisterShapeOperator(*registry);
+	RegisterSliceOperator(*registry);
+	RegisterConcatOperator(*registry);
+	RegisterUnsqueezeOperator(*registry);
 
 	return true;
 }
diff --git a/Engine/Plugins/Experimental/NNE/Source/NNEUtils/Private/NNEUtilsModelBuilderONNX.cpp b/Engine/Plugins/Experimental/NNE/Source/NNEUtils/Private/NNEUtilsModelBuilderONNX.cpp
--- a/Engine/Plugins/Experimental/NNE/Source/NNEUtils/Private/NNEUtilsModelBuilderONNX.cpp
+++ b/Engine/Plugins/Experimental/NNE/Source/NNEUtils/Private/NNEUtilsModelBuilderONNX.cpp
@@ -350,10 +350,11 @@ NNEUTILS_API bool CreateONNXModelForOperator(bool UseVariadicShapeForModel, cons
 	
 	int64 IrVersion = OnnxIrVersion;
 	int64 OpsetVersion = OnnxOpsetVersion;
-	if (OperatorName == TEXT("Upsample") || OperatorName == TEXT("Pad"))
+	if (OperatorName == TEXT("Upsample") ||	// Upsample is deprecated starting opset 10
+		OperatorName == TEXT("Pad") ||		// Pad current implementation is opset 2 (next version is 11)
+		OperatorName == TEXT("Slice") ||	// Slice current implementation is opset 1 (next version is 10)
+		OperatorName == TEXT("Unsqueeze"))	// Pad current implementation is opset 1 (next version is 11)
 	{
-		// Upsample is deprecated starting opset 10
-		// Pad current implementation is opset 2 (next version is 11)
 		OpsetVersion = 9;
 	}
 	TUniquePtr<IModelBuilder> Builder(CreateONNXModelBuilder(IrVersion, OpsetVersion));
