NNE - Add DirectML Transpose

diff --git a/Engine/Plugins/Experimental/NNE/Source/NNEQA/Resources/NNEQATestDesc.json b/Engine/Plugins/Experimental/NNE/Source/NNEQA/Resources/NNEQATestDesc.json
--- a/Engine/Plugins/Experimental/NNE/Source/NNEQA/Resources/NNEQATestDesc.json
+++ b/Engine/Plugins/Experimental/NNE/Source/NNEQA/Resources/NNEQATestDesc.json
@@ -329,6 +329,15 @@
 			"targets": [
 				{ "target": "GRU", "absolute_tolerance": 1e-6 }
 			]
+		},
+		{
+			"category": "TransposeOp",
+			"runtimes": [
+				{ "name": "NNERuntimeRDGHlsl", "skip": true }
+			],
+			"targets": [
+				{ "target": "Transpose" }
+			]
 		}
 	],
 
@@ -778,6 +787,24 @@
 			"datasets": [
 				{ "inputs": [ { "shape": [ 2, 3, 3 ] }, { "shape": [ 1, 15, 3 ] }, { "shape": [ 1, 15, 5 ] }, { "shape": [ 1, 30 ] } ], "outputs": [ { "shape": [ 2, 1, 3, 5 ] }, { "shape": [ 1, 3, 5 ] } ] }
 			]
+		},
+		{ 
+			"name": "TransposeOp.234to243",
+			"datasets": [
+				{ "inputs": [ { "shape": [ 2, 3, 4 ] } ], "outputs": [ { "shape": [ 2, 4, 3 ] } ] }
+			]
+		},
+		{ 
+			"name": "TransposeOp.234to423",
+			"datasets": [
+				{ "inputs": [ { "shape": [ 2, 3, 4 ] } ], "outputs": [ { "shape": [ 4, 2, 3 ] } ] }
+			]
+		},
+		{ 
+			"name": "TransposeOp.Default",
+			"datasets": [
+				{ "inputs": [ { "shape": [ 5, 7, 8 ] } ], "outputs": [ { "shape": [ 8, 7, 5 ] } ] }
+			]
 		}
 	],
 
@@ -1128,6 +1155,28 @@
 					{ "name": "activation_alpha", "value": [0.01, 0.2], "type": "FloatArray" }
 				] }
 			]
+		},
+		{
+			"name": "TransposeOp.234to243",
+			"attribute_maps": [
+				{ "attributes": [
+					{ "name": "perm", "value": [0, 2, 1], "type": "Int32Array" }
+				] }
+			]
+		},
+		{
+			"name": "TransposeOp.234to423",
+			"attribute_maps": [
+				{ "attributes": [
+					{ "name": "perm", "value": [2, 0, 1], "type": "Int32Array" }
+				] }
+			]
+		},
+		{
+			"name": "TransposeOp.Default",
+			"attribute_maps": [
+				{ "attributes": [] }
+			]
 		}
 	]
 
diff --git a/Engine/Plugins/Experimental/NNE/Source/NNERuntimeRDG/Private/Dml/NNEDmlOperator.cpp b/Engine/Plugins/Experimental/NNE/Source/NNERuntimeRDG/Private/Dml/NNEDmlOperator.cpp
--- a/Engine/Plugins/Experimental/NNE/Source/NNERuntimeRDG/Private/Dml/NNEDmlOperator.cpp
+++ b/Engine/Plugins/Experimental/NNE/Source/NNERuntimeRDG/Private/Dml/NNEDmlOperator.cpp
@@ -10,6 +10,29 @@ namespace UE::NNERuntimeRDG::Private::Dml
 
 namespace DmlUtil
 {
+
+	void FTensorDesc::UpdateShapesAndStrides(TConstArrayView<uint32> InSizes, TConstArrayView<uint32> InStrides)
+	{
+		Sizes = InSizes;
+        BuffDesc.Sizes = Sizes.GetData();
+
+		Strides = InStrides;
+        BuffDesc.Strides = Strides.GetData();
+	}
+
+	void FTensorDesc::SetStridesFromFTensor(const NNECore::Internal::FTensor& InputDesc)
+	{
+		uint32 CurrStride = 1;
+
+		Strides.SetNum(InputDesc.GetShape().Rank());
+		
+		for (int32 i = InputDesc.GetShape().Rank() - 1; i >= 0; --i)
+		{
+			Strides[i] = CurrStride;
+			CurrStride *= InputDesc.GetShape().GetData()[i];
+		}
+	}
+
 	bool FTensorDesc::InitFromTensor(const NNECore::Internal::FTensor& Tensor, int32 MinTensorRank, TConstArrayView<uint32> BroadcastShape)
 	{
 		Reset();
@@ -84,6 +107,7 @@ namespace DmlUtil
 		{
 			Sizes.Insert(1, 0);
 		}
+
 	}
 
 	void FTensorDesc::SetShapeAndStrides(TConstArrayView<uint32> Shape, TConstArrayView<uint32> BroadcastShape)
diff --git a/Engine/Plugins/Experimental/NNE/Source/NNERuntimeRDG/Private/Dml/NNEDmlOperator.h b/Engine/Plugins/Experimental/NNE/Source/NNERuntimeRDG/Private/Dml/NNEDmlOperator.h
--- a/Engine/Plugins/Experimental/NNE/Source/NNERuntimeRDG/Private/Dml/NNEDmlOperator.h
+++ b/Engine/Plugins/Experimental/NNE/Source/NNERuntimeRDG/Private/Dml/NNEDmlOperator.h
@@ -45,12 +45,14 @@ public:
 
 namespace DmlUtil
 {
+	using FSmallIntArray = TArray<int32, TInlineAllocator<NNECore::FTensorShape::MaxRank>>;
 	using FSmallUIntArray = TArray<uint32, TInlineAllocator<NNECore::FTensorShape::MaxRank>>;
 
 	struct FTensorDesc
 	{
 		DML_BUFFER_TENSOR_DESC	BuffDesc;
 		DML_TENSOR_DESC			Desc;
+		// Don't edit Sizes and Strides directly, use methods
 		FSmallUIntArray			Sizes;
 		FSmallUIntArray			Strides;
 		uint64					ElemSizeInBytes;
@@ -58,6 +60,10 @@ namespace DmlUtil
 		bool InitFromTensor(const NNECore::Internal::FTensor& InputDesc, int32 MinTensorRank, TConstArrayView<uint32> Broadcast = MakeArrayView((uint32*) nullptr, 0));
 		bool InitFromTensor1D(const NNECore::Internal::FTensor& InputDesc, int32 Rank);
 
+		void SetStridesFromFTensor(const NNECore::Internal::FTensor& InputDesc);
+
+		void UpdateShapesAndStrides(TConstArrayView<uint32> InSizes, TConstArrayView<uint32> InStrides);
+
 	private:
 		void Reset();
 
diff --git a/Engine/Plugins/Experimental/NNE/Source/NNERuntimeRDG/Private/Dml/NNEDmlOperatorTranspose.cpp b/Engine/Plugins/Experimental/NNE/Source/NNERuntimeRDG/Private/Dml/NNEDmlOperatorTranspose.cpp
new file mode 100644
--- /dev/null
+++ b/Engine/Plugins/Experimental/NNE/Source/NNERuntimeRDG/Private/Dml/NNEDmlOperatorTranspose.cpp
@@ -0,0 +1,97 @@
+// Copyright Epic Games, Inc. All Rights Reserved.
+
+#ifdef NNE_USE_DIRECTML
+
+#include "NNEDmlOperator.h"
+#include "Math/Range.h"
+#include "Algo/Transform.h"
+#include "Algo/Reverse.h"
+
+namespace UE::NNERuntimeRDG::Private::Dml
+{
+
+/**
+ * Transpose
+ */
+class FOperatorDmlTranspose : public FOperatorDml
+{
+
+public:
+
+	static FOperatorDml* Create()
+	{
+		return new FOperatorDmlTranspose();
+	}
+
+	//
+	//
+	//
+	virtual bool Initialize(IDMLDevice* Device, TArrayView<const NNECore::Internal::FTensor> InputTensors, TArrayView<const NNECore::Internal::FTensor> OutputTensors, const NNECore::FAttributeMap& Attributes) override
+	{
+        check(InputTensors.Num() == 1);
+        check(OutputTensors.Num() == 1);
+
+        int32 NumDims = InputTensors[0].GetShape().Rank();
+        check(NumDims > 0);
+        check(NumDims == OutputTensors[0].GetShape().Rank());
+
+        // Default permutation is reverse
+        DmlUtil::FSmallIntArray ReversePerm;
+        for(int Idx = NumDims-1; Idx >= 0; Idx--)
+        {
+            ReversePerm.Add(Idx);
+        }
+        TArray<int32> Perm = 
+			Attributes.GetValueOrDefault<TArray<int32>>(TEXT("perm"), (TArray<int32>) ReversePerm);
+
+        check(Perm.Num() == NumDims);
+
+        // Initialize Input tensor desc
+        DmlUtil::FTensorDesc DmlInputTensorDesc;
+        if (!DmlInputTensorDesc.InitFromTensor(InputTensors[0], NumDims))
+        {
+            UE_LOG(LogNNE, Error, TEXT("Failed to initialize Transpose input for DML inference"));
+            return false;
+        }
+        
+        DmlInputTensorDesc.SetStridesFromFTensor(InputTensors[0]);
+
+        // Initialize Output tensor desc
+        DmlUtil::FTensorDesc DmlOutputTensorDesc;
+        if (!DmlOutputTensorDesc.InitFromTensor(OutputTensors[0], NumDims))
+        {
+            UE_LOG(LogNNE, Error, TEXT("Failed to initialize Transpose output for DML inference"));
+            return false;
+        }
+
+        // Apply permutations to both sizes and strides
+        auto PermuteFunc = [&, Perm] (TConstArrayView<uint32> InputView) -> DmlUtil::FSmallUIntArray
+            {
+                DmlUtil::FSmallUIntArray Permuted;
+                for(int32 PermVal : Perm)
+                {
+                    Permuted.Add(InputView[PermVal]);
+                }
+                return Permuted;
+            };
+        
+        DmlInputTensorDesc.UpdateShapesAndStrides(PermuteFunc(DmlInputTensorDesc.Sizes), PermuteFunc(DmlInputTensorDesc.Strides));
+        
+        check(DmlOutputTensorDesc.Sizes == DmlInputTensorDesc.Sizes);        
+
+		DML_ELEMENT_WISE_IDENTITY_OPERATOR_DESC DmlIdentityOpDesc{};
+
+        DmlIdentityOpDesc.InputTensor = &DmlInputTensorDesc.Desc;
+        DmlIdentityOpDesc.OutputTensor = &DmlOutputTensorDesc.Desc;
+
+		return CreateOperator(Device, DML_OPERATOR_DESC{ DML_OPERATOR_ELEMENT_WISE_IDENTITY, &DmlIdentityOpDesc} );
+
+	}
+};
+
+// Register Transpose operator on Module startup
+NNE_DML_REGISTER_OP(Transpose)
+
+} // namespace UE::NNERuntimeRDG::Private::Dml
+
+#endif // NNE_USE_DIRECTML
