NNE HLSL upload of constant tensors from shape inference + Shape operator tests.

Problem:
* Constant tensor from shape inference were not exposed to GPU memory + we didn't had tests for shape operator.

Solutions:
* Upload constant intermediate tensors when model input are prepared, also upload constant data to user provided gpu buffer if an output is const.

Tests: no regression from main for nne test on editor win64 + style transfer demo

diff --git a/Engine/Plugins/Experimental/NNE/Source/NNEQA/Private/NNEQAJsonUtils.h b/Engine/Plugins/Experimental/NNE/Source/NNEQA/Private/NNEQAJsonUtils.h
--- a/Engine/Plugins/Experimental/NNE/Source/NNEQA/Private/NNEQAJsonUtils.h
+++ b/Engine/Plugins/Experimental/NNE/Source/NNEQA/Private/NNEQAJsonUtils.h
@@ -41,13 +41,17 @@ namespace Json
 	struct FTestConfigRuntime : FJsonSerializable
 	{
 		FString Name;
-		bool Skip;
 		float AbsoluteTolerance = JSON_TOLERANCE_NOTSET;
 		float RelativeTolerance = JSON_TOLERANCE_NOTSET;
+		bool Skip;
+		bool SkipStatic;
+		bool SkipVariadic;
 
 		BEGIN_JSON_SERIALIZER
 			JSON_SERIALIZE("name", Name);
 			JSON_SERIALIZE("skip", Skip);
+			JSON_SERIALIZE("skip_static_test", SkipStatic);
+			JSON_SERIALIZE("skip_variadic_test", SkipVariadic);
 			JSON_SERIALIZE("absolute_tolerance", AbsoluteTolerance);
 			JSON_SERIALIZE("relative_tolerance", RelativeTolerance);
 		END_JSON_SERIALIZER
diff --git a/Engine/Plugins/Experimental/NNE/Source/NNEQA/Private/NNEQAParametricTest.cpp b/Engine/Plugins/Experimental/NNE/Source/NNEQA/Private/NNEQAParametricTest.cpp
--- a/Engine/Plugins/Experimental/NNE/Source/NNEQA/Private/NNEQAParametricTest.cpp
+++ b/Engine/Plugins/Experimental/NNE/Source/NNEQA/Private/NNEQAParametricTest.cpp
@@ -83,6 +83,14 @@ private:
 			}
 			else
 			{
+				if (Runtime.SkipStatic)
+				{
+					TestSetup.SkipStaticTestForRuntime.Add(Runtime.Name);
+				}
+				if (Runtime.SkipVariadic)
+				{
+					TestSetup.SkipStaticTestForRuntime.Add(Runtime.Name);
+				}
 				if (Runtime.AbsoluteTolerance != Json::JSON_TOLERANCE_NOTSET)
 				{
 					TestSetup.AbsoluteToleranceForRuntime.Emplace(Runtime.Name, Runtime.AbsoluteTolerance);
diff --git a/Engine/Plugins/Experimental/NNE/Source/NNEQA/Private/NNEQAUtils.cpp b/Engine/Plugins/Experimental/NNE/Source/NNEQA/Private/NNEQAUtils.cpp
--- a/Engine/Plugins/Experimental/NNE/Source/NNEQA/Private/NNEQAUtils.cpp
+++ b/Engine/Plugins/Experimental/NNE/Source/NNEQA/Private/NNEQAUtils.cpp
@@ -365,6 +365,10 @@ namespace UE::NNEQA::Private
 		{
 			return CompareTensorData<int32>(RefTensor, RefRawBuffer, OtherTensor, OtherRawBuffer, AbsoluteTolerance, RelativeTolerance);
 		}
+		else if (RefTensor.GetDataType() == ENNETensorDataType::Int64)
+		{
+			return CompareTensorData<int64>(RefTensor, RefRawBuffer, OtherTensor, OtherRawBuffer, AbsoluteTolerance, RelativeTolerance);
+		}
 		else if (RefTensor.GetDataType() == ENNETensorDataType::UInt32)
 		{
 			return CompareTensorData<uint32>(RefTensor, RefRawBuffer, OtherTensor, OtherRawBuffer, AbsoluteTolerance, RelativeTolerance);
@@ -551,30 +555,39 @@ namespace UE::NNEQA::Private
 			}
 			FString TestResult;
 			if (TestSetup.AutomationExcludedRuntime.Contains(RuntimeName) ||
-				TestSetup.AutomationExcludedPlatformRuntimeCombination.Contains(TPair<FString, FString>(CurrentPlatform, RuntimeName)))
+				TestSetup.AutomationExcludedPlatformRuntimeCombination.Contains(TPair<FString, FString>(CurrentPlatform, RuntimeName)) ||
+				(TestSetup.SkipStaticTestForRuntime.Contains(RuntimeName) && TestSetup.SkipVariadicTestForRuntime.Contains(RuntimeName)))
 			{
 				TestResult = TEXT("skipped (by config)");
 			}
 			else
 			{
-				bool bShouldRunVariadicTest = (ONNXModelVariadic.Format != ENNEInferenceFormat::Invalid);
+				bool bTestSuceeded = true;
+				bool bShouldRunStaticTest = !TestSetup.SkipStaticTestForRuntime.Contains(RuntimeName);
+				bool bShouldRunVariadicTest = !TestSetup.SkipVariadicTestForRuntime.Contains(RuntimeName);
+
+				bShouldRunVariadicTest &= (ONNXModelVariadic.Format != ENNEInferenceFormat::Invalid);
 				bShouldRunVariadicTest &= !(RuntimeName == "NNERuntimeRDGDml");
 				
-				bool bTestSuceeded = RunTestInferenceAndCompareToRef(TestSetup, Runtime->GetRuntimeName(), ONNXModel, RefOutputMemBuffers, RefOutputTensors);
-
-				if (bShouldRunVariadicTest)
+				if (bShouldRunStaticTest)
 				{
-					if (!bTestSuceeded)
+					bool bTestResult = RunTestInferenceAndCompareToRef(TestSetup, Runtime->GetRuntimeName(), ONNXModel, RefOutputMemBuffers, RefOutputTensors);
+					if (!bTestResult)
 					{
 						UE_LOG(LogNNE, Error, TEXT("Failed running static test."));
+						bTestSuceeded = false;
 					}
-					bool bVariadicTestSuceeded = RunTestInferenceAndCompareToRef(TestSetup, Runtime->GetRuntimeName(), ONNXModelVariadic, RefOutputMemBuffers, RefOutputTensors);
-					if (!bVariadicTestSuceeded)
+				}
+				if (bShouldRunVariadicTest)
+				{
+					bool bTestResult = RunTestInferenceAndCompareToRef(TestSetup, Runtime->GetRuntimeName(), ONNXModelVariadic, RefOutputMemBuffers, RefOutputTensors);
+					if (!bTestResult)
 					{
-						bTestSuceeded = false;
 						UE_LOG(LogNNE, Error, TEXT("Failed running variadic test."));
+						bTestSuceeded = false;
 					}
 				}
+				
 				TestResult = bTestSuceeded ? TEXT("SUCCESS") : TEXT("FAILED");
 				bAllTestsSucceeded &= bTestSuceeded;
 			}
diff --git a/Engine/Plugins/Experimental/NNE/Source/NNEQA/Private/NNEQAUtils.h b/Engine/Plugins/Experimental/NNE/Source/NNEQA/Private/NNEQAUtils.h
--- a/Engine/Plugins/Experimental/NNE/Source/NNEQA/Private/NNEQAUtils.h
+++ b/Engine/Plugins/Experimental/NNE/Source/NNEQA/Private/NNEQAUtils.h
@@ -37,6 +37,8 @@ namespace UE::NNEQA::Private
 			bool IsModelTest;
 			TMap<FString, float> AbsoluteToleranceForRuntime;
 			TMap<FString, float> RelativeToleranceForRuntime;
+			TSet<FString> SkipStaticTestForRuntime;
+			TSet<FString> SkipVariadicTestForRuntime;
 			TArray<NNECore::Internal::FTensor> Inputs;
 			TArray<NNECore::Internal::FTensor> Weights;
 			TArray<NNECore::Internal::FTensor> Outputs;
diff --git a/Engine/Plugins/Experimental/NNE/Source/NNEQA/Resources/NNEQATestDesc.json b/Engine/Plugins/Experimental/NNE/Source/NNEQA/Resources/NNEQATestDesc.json
--- a/Engine/Plugins/Experimental/NNE/Source/NNEQA/Resources/NNEQATestDesc.json
+++ b/Engine/Plugins/Experimental/NNE/Source/NNEQA/Resources/NNEQATestDesc.json
@@ -338,6 +338,16 @@
 			"targets": [
 				{ "target": "Transpose" }
 			]
+		},
+		{
+			"category": "ShapeOp",
+			"runtimes": [
+				{ "name": "NNERuntimeRDGDml", "skip": true },
+				{ "name": "NNERuntimeRDGHlsl", "skip_static_test": true  }
+			],
+			"targets": [
+				{ "target": "Shape" }
+			]
 		}
 	],
 
@@ -805,6 +815,14 @@
 			"datasets": [
 				{ "inputs": [ { "shape": [ 5, 7, 8 ] } ], "outputs": [ { "shape": [ 8, 7, 5 ] } ] }
 			]
+		},
+		{
+			"name": "ShapeOp",
+			"datasets": [
+				{ "inputs":  [ { "shape": [ 3 ] } ], "outputs": [ { "shape": [ 1 ], "type": "int64" } ] },
+				{ "inputs":  [ { "shape": [ 1, 128, 50, 55 ] } ], "outputs": [ { "shape": [ 4 ], "type": "int64"} ] },
+				{ "inputs":  [ { "shape": [ 2, 1, 3, 1, 2 ] } ],  "outputs": [ { "shape": [ 5 ], "type": "int64"} ] }
+			]
 		}
 	],
 
diff --git a/Engine/Plugins/Experimental/NNE/Source/NNEQA/Resources/NNEQATestDescSchema.json b/Engine/Plugins/Experimental/NNE/Source/NNEQA/Resources/NNEQATestDescSchema.json
--- a/Engine/Plugins/Experimental/NNE/Source/NNEQA/Resources/NNEQATestDescSchema.json
+++ b/Engine/Plugins/Experimental/NNE/Source/NNEQA/Resources/NNEQATestDescSchema.json
@@ -65,6 +65,14 @@
 					"type": "boolean",
 					"description": "true to skip the test for this runtime, false by default"
 				},
+				"skip_variadic_test": {
+					"type": "boolean",
+					"description": "true to skip variadic input shape tests"
+				},
+				"skip_static_test": {
+					"type": "boolean",
+					"description": "true to skip static intput shape tests"
+				},
 				"absolute_tolerance": {
 					"type": "number",
 					"minimum": 0,
diff --git a/Engine/Plugins/Experimental/NNE/Source/NNERuntimeRDG/Private/Hlsl/NNERuntimeRDGConcat.cpp b/Engine/Plugins/Experimental/NNE/Source/NNERuntimeRDG/Private/Hlsl/NNERuntimeRDGConcat.cpp
--- a/Engine/Plugins/Experimental/NNE/Source/NNERuntimeRDG/Private/Hlsl/NNERuntimeRDGConcat.cpp
+++ b/Engine/Plugins/Experimental/NNE/Source/NNERuntimeRDG/Private/Hlsl/NNERuntimeRDGConcat.cpp
@@ -87,8 +87,7 @@ namespace UE::NNERuntimeRDG::Private::Hlsl
 
 		virtual void Dispatch(FRDGBuilder& GraphBuilder, TConstArrayView<FTensorRDGRef> InputTensors, TConstArrayView<FTensorRDGRef> OutputTensors) override
 		{
-			//Need to upload the data to gpu tensor or have the hlsl model understand that this operator result will only be needed on CPU.
-			//Current behavior will be faulty if the model rely on this operator result on GPU.
+			UE_LOG(LogNNE, Warning, TEXT("Concat: Output should be constant and already uploaded to GPU memory. Dispatch should not need to be called."));
 		}
 	};
 
diff --git a/Engine/Plugins/Experimental/NNE/Source/NNERuntimeRDG/Private/Hlsl/NNERuntimeRDGShape.cpp b/Engine/Plugins/Experimental/NNE/Source/NNERuntimeRDG/Private/Hlsl/NNERuntimeRDGShape.cpp
--- a/Engine/Plugins/Experimental/NNE/Source/NNERuntimeRDG/Private/Hlsl/NNERuntimeRDGShape.cpp
+++ b/Engine/Plugins/Experimental/NNE/Source/NNERuntimeRDG/Private/Hlsl/NNERuntimeRDGShape.cpp
@@ -59,8 +59,7 @@ namespace UE::NNERuntimeRDG::Private::Hlsl
 
 		virtual void Dispatch(FRDGBuilder& GraphBuilder, TConstArrayView<FTensorRDGRef> InputTensors, TConstArrayView<FTensorRDGRef> OutputTensors) override
 		{
-			//Need to upload the data to gpu tensor or have the hlsl model understand that this operator result will only be needed on CPU.
-			//Current behavior will be faulty if the model rely on this operator result on GPU.
+			UE_LOG(LogNNE, Warning, TEXT("Shape: Output should be constant and already uploaded to GPU memory. Dispatch should not need to be called."));
 		}
 	};
 
diff --git a/Engine/Plugins/Experimental/NNE/Source/NNERuntimeRDG/Private/Hlsl/NNERuntimeRDGSlice.cpp b/Engine/Plugins/Experimental/NNE/Source/NNERuntimeRDG/Private/Hlsl/NNERuntimeRDGSlice.cpp
--- a/Engine/Plugins/Experimental/NNE/Source/NNERuntimeRDG/Private/Hlsl/NNERuntimeRDGSlice.cpp
+++ b/Engine/Plugins/Experimental/NNE/Source/NNERuntimeRDG/Private/Hlsl/NNERuntimeRDGSlice.cpp
@@ -86,7 +86,13 @@ namespace UE::NNERuntimeRDG::Private::Hlsl
 			
 			Internal::CPUHelper::Slice::Apply(*InputTensors[0], *OutputTensors[0], Start);
 
-			return -1;
+			if (!OutputTensors[0]->HasPreparedData())
+			{
+				UE_LOG(LogNNE, Warning, TEXT("Slice: Output could not be computed as a constant tensor, however Slice is not implemented on GPU at the moment."));
+				return -1;
+			}
+
+			return 0;
 		};
 
 		virtual bool Initialize(TConstArrayView<NNECore::FTensorDesc> InputTensorDescs, TConstArrayView<NNECore::FTensorDesc> OutputTensorDescs, const NNECore::FAttributeMap& Attributes) override
@@ -117,8 +123,7 @@ namespace UE::NNERuntimeRDG::Private::Hlsl
 
 		virtual void Dispatch(FRDGBuilder& GraphBuilder, TConstArrayView<FTensorRDGRef> InputTensors, TConstArrayView<FTensorRDGRef> OutputTensors) override
 		{
-			//Need to upload the data to gpu tensor or have the hlsl model understand that this operator result will only be needed on CPU.
-			//Current behavior will be faulty if the model rely on this operator result on GPU.
+			UE_LOG(LogNNE, Warning, TEXT("Slice: Output should be constant and already uploaded to GPU memory. Dispatch should not need to be called."));
 		}
 	};
 
diff --git a/Engine/Plugins/Experimental/NNE/Source/NNERuntimeRDG/Private/Hlsl/NNERuntimeRDGUnsqueeze.cpp b/Engine/Plugins/Experimental/NNE/Source/NNERuntimeRDG/Private/Hlsl/NNERuntimeRDGUnsqueeze.cpp
--- a/Engine/Plugins/Experimental/NNE/Source/NNERuntimeRDG/Private/Hlsl/NNERuntimeRDGUnsqueeze.cpp
+++ b/Engine/Plugins/Experimental/NNE/Source/NNERuntimeRDG/Private/Hlsl/NNERuntimeRDGUnsqueeze.cpp
@@ -72,8 +72,7 @@ namespace UE::NNERuntimeRDG::Private::Hlsl
 
 		virtual void Dispatch(FRDGBuilder& GraphBuilder, TConstArrayView<FTensorRDGRef> InputTensors, TConstArrayView<FTensorRDGRef> OutputTensors) override
 		{
-			//Need to upload the data to gpu tensor or have the hlsl model understand that this operator result will only be needed on CPU.
-			//Current behavior will be faulty if the model rely on this operator result on GPU.
+			UE_LOG(LogNNE, Warning, TEXT("Unsqueeze: Output should be constant and already uploaded to GPU memory. Dispatch should not need to be called."));
 		}
 	};
 
diff --git a/Engine/Plugins/Experimental/NNE/Source/NNERuntimeRDG/Private/NNERuntimeRDGModelHlsl.cpp b/Engine/Plugins/Experimental/NNE/Source/NNERuntimeRDG/Private/NNERuntimeRDGModelHlsl.cpp
--- a/Engine/Plugins/Experimental/NNE/Source/NNERuntimeRDG/Private/NNERuntimeRDGModelHlsl.cpp
+++ b/Engine/Plugins/Experimental/NNE/Source/NNERuntimeRDG/Private/NNERuntimeRDGModelHlsl.cpp
@@ -121,13 +121,28 @@ void FModel::AddDispatchOps_RenderThread(FRDGBuilder& GraphBuilder)
 		{
 			InputTensors.Add(AllTensorRDGRefs[i]);
 		}
+		bool AllOutputTensorConstant = true;
 		OutputTensors.Reset(OperatorOutputTensorIndices.Num());
 		for (int32 i : OperatorOutputTensorIndices[Idx])
 		{
+			AllOutputTensorConstant &= AllTensorRDGRefs[i]->HasPreparedData();
 			OutputTensors.Add(AllTensorRDGRefs[i]);
 		}
 
-		Operators[Idx]->Dispatch(GraphBuilder, InputTensors, OutputTensors);
+		//If all output for operator are constant we don't need to run it.
+		if (!AllOutputTensorConstant)
+		{
+			Operators[Idx]->Dispatch(GraphBuilder, InputTensors, OutputTensors);
+		}
+	}
+
+	//If a model output is constant we upload to it (a user provided GPU buffer).
+	for (const FTensorRDG& OutputTensor : OutputTensorRDGs)
+	{
+		if (OutputTensor.HasPreparedData())
+		{
+			GraphBuilder.QueueBufferUpload(OutputTensor.GetBuffer(), OutputTensor.GetPreparedData<uint8>().GetData(), OutputTensor.GetDataSize(), ERDGInitialDataFlags::None);
+		}
 	}
 }
 
@@ -200,60 +215,101 @@ int FModel::PrepareTensorShapesAndData()
 		};
 	);
 
+	if (!PrepareConstants())
+	{
+		UE_LOG(LogNNE, Warning, TEXT("Constant tensors could not be uploaded to GPU memory."));
+		return -1;
+	}
+
 	return 0;
 }
 
-bool FModel::PrepareWeights()
+namespace UploadHelper
 {
-	if (!WeightsExternalRDGResources.IsEmpty())
-	{
-		check(WeightsExternalRDGResources.Num() == WeightTensorRDGs.Num())
-		return true;
-	}
 
-	//Upload to GPU
-	FEvent* Signal = FGenericPlatformProcess::GetSynchEventFromPool(false);
-	WeightsExternalRDGResources.SetNum(WeightTensorRDGs.Num());
+void EnqueueTensorUpload(TArray<TRefCountPtr<FRDGPooledBuffer>>& OutExternalRDGResources,
+	FTensorRDGArray& TensorToUploadRDGs, ERDGInitialDataFlags CopyDataFlag)
+	{
+		OutExternalRDGResources.Reset();
+		OutExternalRDGResources.SetNum(TensorToUploadRDGs.Num());
 
-	ENQUEUE_RENDER_COMMAND(FModel_PrepareWeights)
-	(
-		[&Signal, this](FRHICommandListImmediate& RHICmdList)
+		bool AtLeastOneConstantTensor = false;
+		for (const FTensorRDG& TensorRDG : TensorToUploadRDGs)
 		{
-			TOptional<ERHIPipeline> Pipeline = RHICmdList.GetPipeline();
-			if (Pipeline == ERHIPipeline::None)
+			if (TensorRDG.HasPreparedData())
 			{
-				RHICmdList.SwitchPipeline(ERHIPipeline::Graphics);
+				AtLeastOneConstantTensor = true;
 			}
+		}
+		if (!AtLeastOneConstantTensor)
+		{
+			return;
+		}
 
-			FRDGBuilder	RDGBuilder(RHICmdList);
+		FEvent* Signal = FGenericPlatformProcess::GetSynchEventFromPool(false);
 
-			for (int32 i = 0; i < WeightTensorRDGs.Num(); ++i)
+		ENQUEUE_RENDER_COMMAND(FModel_UploadTensors)
+		(
+			[Signal, &OutExternalRDGResources, &TensorToUploadRDGs, CopyDataFlag](FRHICommandListImmediate& RHICmdList)
 			{
-				FTensorRDG& Tensor = WeightTensorRDGs[i];
-				check(!Tensor.HasBuffer());
-				check(Tensor.HasPreparedData());
-				FRDGBufferDesc BufferDesc = FRDGBufferDesc::CreateBufferDesc(Tensor.GetElemByteSize(), Tensor.GetVolume());
-				const FRDGBufferRef TransientRDGBuffer = RDGBuilder.CreateBuffer(BufferDesc, *Tensor.GetName(), ERDGBufferFlags::None);
-				const uint8* TensorData = Tensor.GetPreparedData<uint8>().GetData();
-				
-				WeightsExternalRDGResources[i] = RDGBuilder.ConvertToExternalBuffer(TransientRDGBuffer);
-				Tensor.SetBuffer(TransientRDGBuffer);
-				RDGBuilder.QueueBufferUpload(TransientRDGBuffer, TensorData, Tensor.GetDataSize(), ERDGInitialDataFlags::NoCopy);
+				TOptional<ERHIPipeline> Pipeline = RHICmdList.GetPipeline();
+				if (Pipeline == ERHIPipeline::None)
+				{
+					RHICmdList.SwitchPipeline(ERHIPipeline::Graphics);
+				}
+
+				FRDGBuilder	RDGBuilder(RHICmdList);
+
+				for (int32 i = 0; i < TensorToUploadRDGs.Num(); ++i)
+				{
+					FTensorRDG& Tensor = TensorToUploadRDGs[i];
+					check(!Tensor.HasBuffer());
+					if (Tensor.HasPreparedData())
+					{
+
+						FRDGBufferDesc BufferDesc = FRDGBufferDesc::CreateBufferDesc(Tensor.GetElemByteSize(), Tensor.GetVolume());
+						const FRDGBufferRef TransientRDGBuffer = RDGBuilder.CreateBuffer(BufferDesc, *Tensor.GetName(), ERDGBufferFlags::None);
+						const uint8* TensorData = Tensor.GetPreparedData<uint8>().GetData();
+
+						OutExternalRDGResources[i] = RDGBuilder.ConvertToExternalBuffer(TransientRDGBuffer);
+						Tensor.SetBuffer(TransientRDGBuffer);
+						RDGBuilder.QueueBufferUpload(TransientRDGBuffer, TensorData, Tensor.GetDataSize(), CopyDataFlag);
+					}
+				}
+
+				RDGBuilder.Execute();
+
+				if (CopyDataFlag == ERDGInitialDataFlags::NoCopy)
+				{
+					//To prevent any problem if model is released before upload is done to the GPU. To be improved.
+					RHICmdList.BlockUntilGPUIdle();
+				}
+
+
+				Signal->Trigger();
 			}
+		);
 
-			RDGBuilder.Execute();
+		Signal->Wait();	// Wait for render thread to finish
 
-			//To prevent any problem if model is released before upload is done to the GPU. To be improved.
-			RHICmdList.BlockUntilGPUIdle();
+		FGenericPlatformProcess::ReturnSynchEventToPool(Signal);
+	}
+}
 
-			Signal->Trigger();
-		}
-	);
+bool FModel::PrepareConstants()
+{
+	// Data is copied so model can be released safely or another upload added to the queue (via SetInputShapes).
+	UploadHelper::EnqueueTensorUpload(ConstantsExternalRDGResources, IntermediateTensorRDGs, ERDGInitialDataFlags::None);
+	
+	return true;
+}
 
-	// We need to wait for render thread to finish
-	Signal->Wait();
+bool FModel::PrepareWeights()
+{
+	check(WeightsExternalRDGResources.IsEmpty());
 
-	FGenericPlatformProcess::ReturnSynchEventToPool(Signal);
+	// Data is not copied. A GPU sync will happens see EnqueueTensorUpload().
+	UploadHelper::EnqueueTensorUpload(WeightsExternalRDGResources, WeightTensorRDGs, ERDGInitialDataFlags::NoCopy);
 
 	return true;
 }
diff --git a/Engine/Plugins/Experimental/NNE/Source/NNERuntimeRDG/Private/NNERuntimeRDGModelHlsl.h b/Engine/Plugins/Experimental/NNE/Source/NNERuntimeRDG/Private/NNERuntimeRDGModelHlsl.h
--- a/Engine/Plugins/Experimental/NNE/Source/NNERuntimeRDG/Private/NNERuntimeRDGModelHlsl.h
+++ b/Engine/Plugins/Experimental/NNE/Source/NNERuntimeRDG/Private/NNERuntimeRDGModelHlsl.h
@@ -27,11 +27,13 @@ protected:
 	virtual void AddDispatchOps_RenderThread(FRDGBuilder& GraphBuilder) override;
 
 	bool PrepareWeights();
+	bool PrepareConstants();
 
 private:
 
 	TArray<FOperatorHlsl*>	Operators;
 	TArray<TRefCountPtr<FRDGPooledBuffer>> WeightsExternalRDGResources;
+	TArray<TRefCountPtr<FRDGPooledBuffer>> ConstantsExternalRDGResources;
 };
 
 } // namespace UE::NNERuntimeRDG::Private::Hlsl
\ No newline at end of file
diff --git a/Engine/Plugins/Experimental/NNE/Source/NNEUtils/Private/NNEUtilsModelBuilderONNX.cpp b/Engine/Plugins/Experimental/NNE/Source/NNEUtils/Private/NNEUtilsModelBuilderONNX.cpp
--- a/Engine/Plugins/Experimental/NNE/Source/NNEUtils/Private/NNEUtilsModelBuilderONNX.cpp
+++ b/Engine/Plugins/Experimental/NNE/Source/NNEUtils/Private/NNEUtilsModelBuilderONNX.cpp
@@ -350,10 +350,11 @@ NNEUTILS_API bool CreateONNXModelForOperator(bool UseVariadicShapeForModel, cons
 	
 	int64 IrVersion = OnnxIrVersion;
 	int64 OpsetVersion = OnnxOpsetVersion;
-	if (OperatorName == TEXT("Upsample") ||	// Upsample is deprecated starting opset 10
-		OperatorName == TEXT("Pad") ||		// Pad current implementation is opset 2 (next version is 11)
-		OperatorName == TEXT("Slice") ||	// Slice current implementation is opset 1 (next version is 10)
-		OperatorName == TEXT("Unsqueeze"))	// Pad current implementation is opset 1 (next version is 11)
+	if (OperatorName == TEXT("Upsample") ||	// deprecated starting opset 10
+		OperatorName == TEXT("Pad") ||		// current implementation is opset 2 (next version is 11)
+		OperatorName == TEXT("Slice") ||	// current implementation is opset 1 (next version is 10)
+		OperatorName == TEXT("Shape") ||	// current implementation is opset 1 (next version is 13)
+		OperatorName == TEXT("Unsqueeze"))	// current implementation is opset 1 (next version is 11)
 	{
 		OpsetVersion = 9;
 	}
