NNE : Add HLSL Cast operator for shape inference

Test:
* Added unit test for cast operator shape inference helper.

#rnx

diff --git a/Engine/Plugins/Experimental/NNE/Source/NNEQA/Private/NNEQAUnitTestHelper.cpp b/Engine/Plugins/Experimental/NNE/Source/NNEQA/Private/NNEQAUnitTestHelper.cpp
--- a/Engine/Plugins/Experimental/NNE/Source/NNEQA/Private/NNEQAUnitTestHelper.cpp
+++ b/Engine/Plugins/Experimental/NNE/Source/NNEQA/Private/NNEQAUnitTestHelper.cpp
@@ -34,16 +34,32 @@ void FShapeInferenceHelperUnitTestBase::GetTests(TArray<FString>& OutBeautifiedN
 
 bool FShapeInferenceHelperUnitTestBase::RunTest(const FString& Parameter) { return false; }
 
-FTensor FShapeInferenceHelperUnitTestBase::MakeTensor(FString Name, TConstArrayView<uint32> Shape)
+FTensor FShapeInferenceHelperUnitTestBase::MakeTensor(FString Name, TConstArrayView<uint32> Shape, ENNETensorDataType DataType)
 {
 	FTensorShape TensorShape = FTensorShape::Make(Shape);
-	FTensor Tensor = FTensor::Make(Name, TensorShape, ENNETensorDataType::Float);
+	FTensor Tensor = FTensor::Make(Name, TensorShape, DataType);
 	return Tensor;
 }
 
 FTensor FShapeInferenceHelperUnitTestBase::MakeConstTensor(FString Name, TConstArrayView<uint32> Shape, TConstArrayView<float> Data)
 {
-	FTensor Tensor = MakeTensor(Name, Shape);
+	FTensor Tensor = MakeTensor(Name, Shape, ENNETensorDataType::Float);
+	check(Data.Num() == Tensor.GetVolume());
+	Tensor.SetPreparedData(Data);
+	return Tensor;
+}
+
+FTensor FShapeInferenceHelperUnitTestBase::MakeConstTensorInt32(FString Name, TConstArrayView<uint32> Shape, TConstArrayView<int32> Data)
+{
+	FTensor Tensor = MakeTensor(Name, Shape, ENNETensorDataType::Int32);
+	check(Data.Num() == Tensor.GetVolume());
+	Tensor.SetPreparedData(Data);
+	return Tensor;
+}
+
+FTensor FShapeInferenceHelperUnitTestBase::MakeConstTensorInt64(FString Name, TConstArrayView<uint32> Shape, TConstArrayView<int64> Data)
+{
+	FTensor Tensor = MakeTensor(Name, Shape, ENNETensorDataType::Int64);
 	check(Data.Num() == Tensor.GetVolume());
 	Tensor.SetPreparedData(Data);
 	return Tensor;
diff --git a/Engine/Plugins/Experimental/NNE/Source/NNEQA/Private/NNEQAUnitTestHelper.h b/Engine/Plugins/Experimental/NNE/Source/NNEQA/Private/NNEQAUnitTestHelper.h
--- a/Engine/Plugins/Experimental/NNE/Source/NNEQA/Private/NNEQAUnitTestHelper.h
+++ b/Engine/Plugins/Experimental/NNE/Source/NNEQA/Private/NNEQAUnitTestHelper.h
@@ -34,11 +34,12 @@ namespace UE::NNEQA::Private::NNERuntimeRDG
 		virtual bool RunTest(const FString& Parameter) override;
 
 		//Shape inference unit test helpers
-		static FTensor MakeTensor(FString Name, TConstArrayView<uint32> Shape);
+		static FTensor MakeTensor(FString Name, TConstArrayView<uint32> Shape, ENNETensorDataType DataType = ENNETensorDataType::Float);
 		static FTensor MakeConstTensor(FString Name, TConstArrayView<uint32> Shape, TConstArrayView<float> Data);
+		static FTensor MakeConstTensorInt32(FString Name, TConstArrayView<uint32> Shape, TConstArrayView<int32> Data);
+		static FTensor MakeConstTensorInt64(FString Name, TConstArrayView<uint32> Shape, TConstArrayView<int64> Data);
 		bool TestUnaryOutputIsOnlyComputedWhenItShould(EElementWiseUnaryOperatorType OpType);
 		bool TestBinaryOutputIsOnlyComputedWhenItShould(EElementWiseBinaryOperatorType OpType);
-		
 	};
 
 	#define IMPLEMENT_NNE_SHAPEINFERENCEHELPER_UNIT_AUTOMATION_TEST( TClass, PrettyName ) \
diff --git a/Engine/Plugins/Experimental/NNE/Source/NNEQA/Private/NNEQAUnitTestOpCast.cpp b/Engine/Plugins/Experimental/NNE/Source/NNEQA/Private/NNEQAUnitTestOpCast.cpp
new file mode 100644
--- /dev/null
+++ b/Engine/Plugins/Experimental/NNE/Source/NNEQA/Private/NNEQAUnitTestOpCast.cpp
@@ -0,0 +1,161 @@
+// Copyright Epic Games, Inc. All Rights Reserved.
+
+#include "CoreMinimal.h"
+#include "NNECoreTensor.h"
+#include "NNECoreTypes.h"
+#include "NNEQAUnitTestHelper.h"
+#include "NNERuntimeRDGHelperCast.h"
+
+#if WITH_DEV_AUTOMATION_TESTS
+#include "Misc/AutomationTest.h"
+#endif
+
+namespace UE::NNEQA::Private::NNERuntimeRDG::CastOp
+{
+
+#if WITH_DEV_AUTOMATION_TESTS
+
+	using namespace NNECore::Internal;
+	using namespace UE::NNERuntimeRDG::Internal;
+
+	IMPLEMENT_NNE_SHAPEINFERENCEHELPER_UNIT_AUTOMATION_TEST(FCastCPUHelperConstOuput, "System.Engine.MachineLearning.NNE.UnitTest.CastHelper.ConstOutput");
+	bool FCastCPUHelperConstOuput::RunTest(const FString& Parameter)
+	{
+		FTensor XC1 = MakeConstTensor(TEXT("XC1"), { 1 }, { 1.0f });
+		FTensor XC20 = MakeConstTensor(TEXT("XC20"), { 20 }, { 3.0f, 4.0f, 3.0f, 4.0f, 3.0f, 3.0f, 4.0f, 3.0f, 4.0f, 3.0f, 3.0f, 4.0f, 3.0f, 4.0f, 3.0f, 3.0f, 4.0f, 3.0f, 4.0f, 3.0f });
+		FTensor X1 = MakeTensor(TEXT("X"), { 1 });
+
+		FTensor Y = MakeTensor(TEXT("Y"), { 1 });
+		CPUHelper::Cast::Apply(XC1 , Y);
+		UTEST_TRUE(TEXT("Y const if input is const"), Y.HasPreparedData());
+
+		Y = MakeTensor(TEXT("Y"), { 1 });
+		CPUHelper::Cast::Apply(X1, Y);
+		UTEST_FALSE(TEXT("Y not const if input not const"), Y.HasPreparedData());
+
+		Y = MakeTensor(TEXT("Y"), { 20 });
+		CPUHelper::Cast::Apply(XC20, Y);
+		UTEST_FALSE(TEXT("Y not const if input is too large "), Y.HasPreparedData());
+
+		return true;
+	}
+
+	IMPLEMENT_NNE_SHAPEINFERENCEHELPER_UNIT_AUTOMATION_TEST(FCastCPUHelperRank1FromFloat, "System.Engine.MachineLearning.NNE.UnitTest.CastHelper.Rank1FromFloat");
+	bool FCastCPUHelperRank1FromFloat::RunTest(const FString& Parameter)
+	{
+		FTensor XC2 = MakeConstTensor(TEXT("XC2"), { 2 }, { -1.0f, 1.5f });
+
+		FTensor Y = MakeTensor(TEXT("Y"), { 2 }, ENNETensorDataType::Float);
+		CPUHelper::Cast::Apply(XC2, Y);
+		UTEST_TRUE(TEXT("Y const if input is const"), Y.HasPreparedData());
+		UTEST_EQUAL(TEXT("Cast(float->float)[0]"), Y.GetPreparedData<float>()[0], -1.0f);
+		UTEST_EQUAL(TEXT("Cast(float->float)[1]"), Y.GetPreparedData<float>()[1], 1.5f);
+
+		Y = MakeTensor(TEXT("Y"), { 2 }, ENNETensorDataType::Int32);
+		CPUHelper::Cast::Apply(XC2, Y);
+		UTEST_TRUE(TEXT("Y const if input is const"), Y.HasPreparedData());
+		UTEST_EQUAL(TEXT("Cast(float->int32)[0]"), Y.GetPreparedData<int32>()[0], (int32)-1);
+		UTEST_EQUAL(TEXT("Cast(float->int32)[1]"), Y.GetPreparedData<int32>()[1], (int32)1);
+
+		Y = MakeTensor(TEXT("Y"), { 2 }, ENNETensorDataType::Int64);
+		CPUHelper::Cast::Apply(XC2, Y);
+		UTEST_TRUE(TEXT("Y const if input is const"), Y.HasPreparedData());
+		UTEST_EQUAL(TEXT("Cast(float->int64)[0]"), Y.GetPreparedData<int64>()[0], (int64)-1);
+		UTEST_EQUAL(TEXT("Cast(float->int64)[1]"), Y.GetPreparedData<int64>()[1], (int64)1);
+
+		return true;
+	}
+
+	IMPLEMENT_NNE_SHAPEINFERENCEHELPER_UNIT_AUTOMATION_TEST(FCastCPUHelperRank1FromInt32, "System.Engine.MachineLearning.NNE.UnitTest.CastHelper.Rank1FromInt32");
+	bool FCastCPUHelperRank1FromInt32::RunTest(const FString& Parameter)
+	{
+		FTensor XC2 = MakeConstTensorInt32(TEXT("XC2"), { 2 }, { -2, 0 });
+
+		FTensor Y = MakeTensor(TEXT("Y"), { 2 }, ENNETensorDataType::Float);
+		CPUHelper::Cast::Apply(XC2, Y);
+		UTEST_TRUE(TEXT("Y const if input is const"), Y.HasPreparedData());
+		UTEST_EQUAL(TEXT("Cast(int32->float)[0]"), Y.GetPreparedData<float>()[0], -2.0f);
+		UTEST_EQUAL(TEXT("Cast(int32->float)[1]"), Y.GetPreparedData<float>()[1], 0.0f);
+
+		Y = MakeTensor(TEXT("Y"), { 2 }, ENNETensorDataType::Int32);
+		CPUHelper::Cast::Apply(XC2, Y);
+		UTEST_TRUE(TEXT("Y const if input is const"), Y.HasPreparedData());
+		UTEST_EQUAL(TEXT("Cast(int32->int32)[0]"), Y.GetPreparedData<int32>()[0], (int32)-2);
+		UTEST_EQUAL(TEXT("Cast(int32->int32)[1]"), Y.GetPreparedData<int32>()[1], (int32)0);
+
+		Y = MakeTensor(TEXT("Y"), { 2 }, ENNETensorDataType::Int64);
+		CPUHelper::Cast::Apply(XC2, Y);
+		UTEST_TRUE(TEXT("Y const if input is const"), Y.HasPreparedData());
+		UTEST_EQUAL(TEXT("Cast(int32->int64)[0]"), Y.GetPreparedData<int64>()[0], (int64)-2);
+		UTEST_EQUAL(TEXT("Cast(int32->int64)[1]"), Y.GetPreparedData<int64>()[1], (int64)0);
+
+		return true;
+	}
+
+	IMPLEMENT_NNE_SHAPEINFERENCEHELPER_UNIT_AUTOMATION_TEST(FCastCPUHelperRank1FromInt64, "System.Engine.MachineLearning.NNE.UnitTest.CastHelper.Rank1FromInt64");
+	bool FCastCPUHelperRank1FromInt64::RunTest(const FString& Parameter)
+	{
+		FTensor XC2 = MakeConstTensorInt64(TEXT("XC2"), { 2 }, { (int64)-4, (int64)10});
+
+		FTensor Y = MakeTensor(TEXT("Y"), { 2 }, ENNETensorDataType::Float);
+		CPUHelper::Cast::Apply(XC2, Y);
+		UTEST_TRUE(TEXT("Y const if input is const"), Y.HasPreparedData());
+		UTEST_EQUAL(TEXT("Cast(int32->float)[0]"), Y.GetPreparedData<float>()[0], -4.0f);
+		UTEST_EQUAL(TEXT("Cast(int32->float)[1]"), Y.GetPreparedData<float>()[1], 10.0f);
+
+		Y = MakeTensor(TEXT("Y"), { 2 }, ENNETensorDataType::Int32);
+		CPUHelper::Cast::Apply(XC2, Y);
+		UTEST_TRUE(TEXT("Y const if input is const"), Y.HasPreparedData());
+		UTEST_EQUAL(TEXT("Cast(int32->int32)[0]"), Y.GetPreparedData<int32>()[0], (int32)-4);
+		UTEST_EQUAL(TEXT("Cast(int32->int32)[1]"), Y.GetPreparedData<int32>()[1], (int32)10);
+
+		Y = MakeTensor(TEXT("Y"), { 2 }, ENNETensorDataType::Int64);
+		CPUHelper::Cast::Apply(XC2, Y);
+		UTEST_TRUE(TEXT("Y const if input is const"), Y.HasPreparedData());
+		UTEST_EQUAL(TEXT("Cast(int32->int64)[0]"), Y.GetPreparedData<int64>()[0], (int64)-4);
+		UTEST_EQUAL(TEXT("Cast(int32->int64)[1]"), Y.GetPreparedData<int64>()[1], (int64)10);
+
+		return true;
+	}
+
+	IMPLEMENT_NNE_SHAPEINFERENCEHELPER_UNIT_AUTOMATION_TEST(FCastCPUHelperRank3, "System.Engine.MachineLearning.NNE.UnitTest.CastHelper.Rank3");
+	bool FCastCPUHelperRank3::RunTest(const FString& Parameter)
+	{
+		FTensor XC1x2x3 = MakeConstTensor(TEXT("XC1x2x3"), { 1,2,3 }, { -1.0f, -0.5f, 0.0f, 0.1f, 0.6f, 1.5f });
+
+		FTensor Y = MakeTensor(TEXT("Y"), { 1,2,3 }, ENNETensorDataType::Float);
+		CPUHelper::Cast::Apply(XC1x2x3, Y);
+		UTEST_TRUE(TEXT("Y const if input is const"), Y.HasPreparedData());
+		UTEST_EQUAL(TEXT("Cast(float,float)[0]"), Y.GetPreparedData<float>()[0], -1.0f);
+		UTEST_EQUAL(TEXT("Cast(float,float)[1]"), Y.GetPreparedData<float>()[1], -0.5f);
+		UTEST_EQUAL(TEXT("Cast(float,float)[2]"), Y.GetPreparedData<float>()[2],  0.0f);
+		UTEST_EQUAL(TEXT("Cast(float,float)[3]"), Y.GetPreparedData<float>()[3],  0.1f);
+		UTEST_EQUAL(TEXT("Cast(float,float)[4]"), Y.GetPreparedData<float>()[4],  0.6f);
+		UTEST_EQUAL(TEXT("Cast(float,float)[5]"), Y.GetPreparedData<float>()[5],  1.5f);
+
+		Y = MakeTensor(TEXT("Y"), { 1,2,3 }, ENNETensorDataType::Int32);
+		CPUHelper::Cast::Apply(XC1x2x3, Y);
+		UTEST_TRUE(TEXT("Y const if input is const"), Y.HasPreparedData());
+		UTEST_EQUAL(TEXT("Cast(float,int32)[0]"), Y.GetPreparedData<int32>()[0], -1);
+		UTEST_EQUAL(TEXT("Cast(float,int32)[1]"), Y.GetPreparedData<int32>()[1], 0);
+		UTEST_EQUAL(TEXT("Cast(float,int32)[2]"), Y.GetPreparedData<int32>()[2], 0);
+		UTEST_EQUAL(TEXT("Cast(float,int32)[3]"), Y.GetPreparedData<int32>()[3], 0);
+		UTEST_EQUAL(TEXT("Cast(float,int32)[4]"), Y.GetPreparedData<int32>()[4], 0);
+		UTEST_EQUAL(TEXT("Cast(float,int32)[5]"), Y.GetPreparedData<int32>()[5], 1);
+
+		Y = MakeTensor(TEXT("Y"), { 1,2,3 }, ENNETensorDataType::Int64);
+		CPUHelper::Cast::Apply(XC1x2x3, Y);
+		UTEST_TRUE(TEXT("Y const if input is const"), Y.HasPreparedData());
+		UTEST_EQUAL(TEXT("Cast(float,int64)[0]"), Y.GetPreparedData<int64>()[0], (int64)-1);
+		UTEST_EQUAL(TEXT("Cast(float,int64)[1]"), Y.GetPreparedData<int64>()[1], (int64)0);
+		UTEST_EQUAL(TEXT("Cast(float,int64)[2]"), Y.GetPreparedData<int64>()[2], (int64)0);
+		UTEST_EQUAL(TEXT("Cast(float,int64)[3]"), Y.GetPreparedData<int64>()[3], (int64)0);
+		UTEST_EQUAL(TEXT("Cast(float,int64)[4]"), Y.GetPreparedData<int64>()[4], (int64)0);
+		UTEST_EQUAL(TEXT("Cast(float,int64)[5]"), Y.GetPreparedData<int64>()[5], (int64)1);
+
+		return true;
+	}
+
+#endif //WITH_DEV_AUTOMATION_TESTS
+
+} // namespace UE::NNEQA::Private::NNERuntimeRDG::CastOp
diff --git a/Engine/Plugins/Experimental/NNE/Source/NNERuntimeRDG/Internal/NNERuntimeRDGHelperCast.h b/Engine/Plugins/Experimental/NNE/Source/NNERuntimeRDG/Internal/NNERuntimeRDGHelperCast.h
new file mode 100644
--- /dev/null
+++ b/Engine/Plugins/Experimental/NNE/Source/NNERuntimeRDG/Internal/NNERuntimeRDGHelperCast.h
@@ -0,0 +1,13 @@
+// Copyright Epic Games, Inc. All Rights Reserved.
+
+#pragma once 
+
+#include "CoreMinimal.h"
+
+namespace UE::NNECore::Internal { class FTensor; }
+
+namespace UE::NNERuntimeRDG::Internal::CPUHelper::Cast
+{
+	void NNERUNTIMERDG_API Apply(const NNECore::Internal::FTensor& Tensor, NNECore::Internal::FTensor& OutputTensor);
+
+} // UE::NNERuntimeRDG::Internal::CPUHelper::Cast
\ No newline at end of file
diff --git a/Engine/Plugins/Experimental/NNE/Source/NNERuntimeRDG/Private/Helper/NNERuntimeRDGHelperCast.cpp b/Engine/Plugins/Experimental/NNE/Source/NNERuntimeRDG/Private/Helper/NNERuntimeRDGHelperCast.cpp
new file mode 100644
--- /dev/null
+++ b/Engine/Plugins/Experimental/NNE/Source/NNERuntimeRDG/Private/Helper/NNERuntimeRDGHelperCast.cpp
@@ -0,0 +1,76 @@
+// Copyright Epic Games, Inc. All Rights Reserved.
+
+#include "NNERuntimeRDGHelperCast.h"
+#include "NNECoreTensor.h"
+
+namespace UE::NNERuntimeRDG::Internal::CPUHelper::Cast
+{
+	template<typename TIn, typename TOut> void ApplyWithResolvedTypes(TConstArrayView<TIn> InArray, TArray<TOut>& OutArray)
+	{
+		check(OutArray.IsEmpty());
+		OutArray.Reset(InArray.Num());
+
+		for (TIn e : InArray)
+		{
+			OutArray.Add((TOut)e);
+		}
+	}
+
+	template<typename TOutput> void ApplyWithResolvedOutputType(const NNECore::Internal::FTensor& Tensor, NNECore::Internal::FTensor& OutputTensor)
+	{
+		ENNETensorDataType SourceType = Tensor.GetDataType();
+		TArray<TOutput> OutputData;
+
+		switch (SourceType)
+		{
+			case ENNETensorDataType::Float:
+			{
+				ApplyWithResolvedTypes(Tensor.GetPreparedData<float>(), OutputData);
+				break;
+			}
+			case ENNETensorDataType::Int32:
+			{
+				ApplyWithResolvedTypes(Tensor.GetPreparedData<int32>(), OutputData);
+				break;
+			}
+			case ENNETensorDataType::Int64:
+			{
+				ApplyWithResolvedTypes(Tensor.GetPreparedData<int64>(), OutputData);
+				break;
+			}
+		}
+
+		OutputTensor.SetPreparedData<TOutput>(OutputData);
+	}
+
+	void Apply(const NNECore::Internal::FTensor& Tensor, NNECore::Internal::FTensor& OutputTensor)
+	{
+		//Heuristic to avoid unexpected performance hit. This helper being intended for shape related arithmetic only.
+		static constexpr int32 MaxItemInInputTensors = NNECore::FTensorShape::MaxRank * 2;
+
+		if (Tensor.HasPreparedData() && (Tensor.GetVolume() <= MaxItemInInputTensors))
+		{
+
+			ENNETensorDataType DestinationType = OutputTensor.GetDataType();
+			switch (DestinationType)
+			{
+			case ENNETensorDataType::Float:
+			{
+				ApplyWithResolvedOutputType<float>(Tensor, OutputTensor);
+				break;
+			}
+			case ENNETensorDataType::Int32:
+			{
+				ApplyWithResolvedOutputType<int32>(Tensor, OutputTensor);
+				break;
+			}
+			case ENNETensorDataType::Int64:
+			{
+				ApplyWithResolvedOutputType<int64>(Tensor, OutputTensor);
+				break;
+			}
+			}
+		}
+	}
+	
+} // UE::NNERuntimeRDG::Internal::CPUHelper::Cast
\ No newline at end of file
diff --git a/Engine/Plugins/Experimental/NNE/Source/NNERuntimeRDG/Private/Hlsl/NNERuntimeRDGCast.cpp b/Engine/Plugins/Experimental/NNE/Source/NNERuntimeRDG/Private/Hlsl/NNERuntimeRDGCast.cpp
new file mode 100644
--- /dev/null
+++ b/Engine/Plugins/Experimental/NNE/Source/NNERuntimeRDG/Private/Hlsl/NNERuntimeRDGCast.cpp
@@ -0,0 +1,168 @@
+// Copyright Epic Games, Inc. All Rights Reserved.
+
+#include "NNERuntimeRDGCast.h"
+#include "NNERuntimeRDGHelperCast.h"
+#include "NNECoreTensor.h"
+#include "NNECoreTypes.h"
+
+namespace UE::NNERuntimeRDG::Private::Hlsl
+{
+	/**
+	 * Cast operator implementation
+	 */
+	class FCast : public FOperatorHlsl
+	{
+	public:
+
+		FCast() {}
+		virtual ~FCast() = default;
+
+		virtual int PrepareOutputs(TConstArrayView<NNECore::Internal::FTensorRef> InputTensors, TArrayView<NNECore::Internal::FTensorRef> OutputTensors) const override
+		{
+			check(InputTensors.Num() == 1);
+			check(OutputTensors.Num() == 1);
+			OutputTensors[0]->SetShape(InputTensors[0]->GetShape());
+
+			const NNECore::Internal::FTensor& X = *InputTensors[0];
+
+			Internal::CPUHelper::Cast::Apply(X, *OutputTensors[0]);
+
+			if (!OutputTensors[0]->HasPreparedData())
+			{
+				UE_LOG(LogNNE, Warning, TEXT("Cast: Output could not be computed as a constant tensor, however Cast is not implemented on GPU at the moment."));
+				return -1;
+			}
+
+			return 0;
+		}
+
+		virtual bool Initialize(TConstArrayView<NNECore::FTensorDesc> InputTensorDescs, TConstArrayView<NNECore::FTensorDesc> OutputTensorDescs, const NNECore::FAttributeMap& Attributes) override
+		{
+			check(InputTensorDescs.Num() == 1);
+			check(OutputTensorDescs.Num() == 1);
+
+			ENNETensorDataType ToFromAttribute = (ENNETensorDataType)Attributes.GetValue<int>(TEXT("to"));
+			ENNETensorDataType ToFromTensor = OutputTensorDescs[0].GetDataType();
+
+			if (ToFromAttribute != ToFromTensor)
+			{
+				UE_LOG(LogNNE, Warning, TEXT("Cast should output a tensor of type %d but was of type %d."), ToFromAttribute, ToFromTensor);
+				return false;
+			}
+			
+			return true;
+		}
+
+		virtual void Dispatch(FRDGBuilder& GraphBuilder, TConstArrayView<FTensorRDGRef> InputTensors, TConstArrayView<FTensorRDGRef> OutputTensors) override
+		{
+			UE_LOG(LogNNE, Warning, TEXT("Cast: Output should be constant and already uploaded to GPU memory. Dispatch should not need to be called."));
+		}
+	};
+
+	bool ValidateCastOperator(const NNECore::FAttributeMap& AttributeMap, TConstArrayView<ENNETensorDataType> InputTypes, TConstArrayView<NNECore::FSymbolicTensorShape> InputShapes)
+	{
+		bool bIsValid = true;
+
+		//This match version 13 of the Cast operator
+		//[REDACTED_PROFILE_URL]/onnx/blob/main/docs/Operators.md#Cast
+		FAttributeValidator AttributeValidator;
+		AttributeValidator.AddRequired(TEXT("to"), ENNEAttributeDataType::Int32);
+		bIsValid &= AttributeValidator.Validate(AttributeMap);
+
+		if (bIsValid)
+		{
+			//In ONNX "To" is assumed to be DataType from TensorProto. Here it is a ENNETensorDataType however no conversion was needed as both enum match.
+			//See ENNETensorDataType GetDataTypeFromGraphTensor(Ort::GraphTensorDataType TensorDataType) in NNEUtilsModelOptimizer.cpp
+			//If a conversion is needed in the future it should be done when converting the model from ONNX to RDG format in the model builder.
+			ENNETensorDataType To = (ENNETensorDataType)AttributeMap.GetValue<int>(TEXT("to"));
+			switch (To)
+			{
+				case ENNETensorDataType::None:
+					UE_LOG(LogNNE, Warning, TEXT("Cast: Invalid target type 'None' not supported."));
+					bIsValid = false;
+					break;
+				case ENNETensorDataType::Char:
+					UE_LOG(LogNNE, Warning, TEXT("Cast: Invalid target type 'Char' not supported."));
+					bIsValid = false;
+					break;
+				case ENNETensorDataType::Boolean:
+					UE_LOG(LogNNE, Warning, TEXT("Cast: Invalid target type 'Boolean' not supported."));
+					bIsValid = false;
+					break;
+				case ENNETensorDataType::Half:
+					UE_LOG(LogNNE, Warning, TEXT("Cast: Invalid target type 'Half' not supported."));
+					bIsValid = false;
+					break;
+				case ENNETensorDataType::Float:
+					break;
+				case ENNETensorDataType::Double:
+					UE_LOG(LogNNE, Warning, TEXT("Cast: Invalid target type 'Double' not supported."));
+					bIsValid = false;
+					break;
+				case ENNETensorDataType::Int8:
+					UE_LOG(LogNNE, Warning, TEXT("Cast: Invalid target type 'Int8' not supported."));
+					bIsValid = false;
+					break;
+				case ENNETensorDataType::Int16:
+					UE_LOG(LogNNE, Warning, TEXT("Cast: Invalid target type 'Int16' not supported."));
+					bIsValid = false;
+					break;
+				case ENNETensorDataType::Int32:
+					break;
+				case ENNETensorDataType::Int64:
+					break;
+				case ENNETensorDataType::UInt8:
+					UE_LOG(LogNNE, Warning, TEXT("Cast: Invalid target type 'UInt8' not supported."));
+					bIsValid = false;
+					break;
+				case ENNETensorDataType::UInt16:
+					UE_LOG(LogNNE, Warning, TEXT("Cast: Invalid target type 'UInt16' not supported."));
+					bIsValid = false;
+					break;
+				case ENNETensorDataType::UInt32:
+					UE_LOG(LogNNE, Warning, TEXT("Cast: Invalid target type 'UInt32' not supported."));
+					bIsValid = false;
+					break;
+				case ENNETensorDataType::UInt64:
+					UE_LOG(LogNNE, Warning, TEXT("Cast: Invalid target type 'UInt64' not supported."));
+					bIsValid = false;
+					break;
+				case ENNETensorDataType::Complex64:
+					UE_LOG(LogNNE, Warning, TEXT("Cast: Invalid target type 'Complex64' not supported."));
+					bIsValid = false;
+					break;
+				case ENNETensorDataType::Complex128:
+					UE_LOG(LogNNE, Warning, TEXT("Cast: Invalid target type 'Complex128' not supported."));
+					bIsValid = false;
+					break;
+				case ENNETensorDataType::BFloat16:
+					UE_LOG(LogNNE, Warning, TEXT("Cast: Invalid target type 'BFloat16' not supported."));
+					bIsValid = false;
+					break;
+				default:
+					UE_LOG(LogNNE, Warning, TEXT("Cast: Invalid target type %d not supported."), To);
+					bIsValid = false;
+			}
+		}
+		
+		FInputValidator InputValidator;
+		InputValidator.AddSupportedType(ENNETensorDataType::Float);
+		InputValidator.AddSupportedType(ENNETensorDataType::Int32);
+		InputValidator.AddSupportedType(ENNETensorDataType::Int64);
+		InputValidator.AddRequired();
+		bIsValid &= InputValidator.Validate(InputTypes);
+
+		return bIsValid;
+	}
+
+	FOperatorHlsl* CreateCastOperator()
+	{
+		return new FCast();
+	}
+
+	bool RegisterCastOperator(FOperatorRegistryHlsl& Registry)
+	{
+		Registry.OpAdd(TEXT("Cast"), CreateCastOperator, ValidateCastOperator);
+		return true;
+	}
+} // UE::NNERuntimeRDG::Private::Hlsl
diff --git a/Engine/Plugins/Experimental/NNE/Source/NNERuntimeRDG/Private/Hlsl/NNERuntimeRDGCast.h b/Engine/Plugins/Experimental/NNE/Source/NNERuntimeRDG/Private/Hlsl/NNERuntimeRDGCast.h
new file mode 100644
--- /dev/null
+++ b/Engine/Plugins/Experimental/NNE/Source/NNERuntimeRDG/Private/Hlsl/NNERuntimeRDGCast.h
@@ -0,0 +1,11 @@
+// Copyright Epic Games, Inc. All Rights Reserved.
+
+#pragma once
+
+#include "NNERuntimeRDG.h"
+#include "NNERuntimeRDGHlslOp.h"
+
+namespace UE::NNERuntimeRDG::Private::Hlsl
+{
+	bool RegisterCastOperator(FOperatorRegistryHlsl& Registry);
+} // UE::NNERuntimeRDG::Private::Hlsl
diff --git a/Engine/Plugins/Experimental/NNE/Source/NNERuntimeRDG/Private/NNERuntimeRDGHlsl.cpp b/Engine/Plugins/Experimental/NNE/Source/NNERuntimeRDG/Private/NNERuntimeRDGHlsl.cpp
--- a/Engine/Plugins/Experimental/NNE/Source/NNERuntimeRDG/Private/NNERuntimeRDGHlsl.cpp
+++ b/Engine/Plugins/Experimental/NNE/Source/NNERuntimeRDG/Private/NNERuntimeRDGHlsl.cpp
@@ -8,7 +8,7 @@
 #include "NNERuntimeRDGHlslHelper.h"
 #include "NNERuntimeRDGModelHlsl.h"
 #include "NNEUtilsModelOptimizer.h"
-#include "Hlsl/NNERuntimeRDGConv.h"
+#include "Hlsl/NNERuntimeRDGCast.h"
 #include "Hlsl/NNERuntimeRDGConv.h"
 #include "Hlsl/NNERuntimeRDGConcat.h"
 #include "Hlsl/NNERuntimeRDGConvTranspose.h"
@@ -48,6 +48,7 @@ bool UNNERuntimeRDGHlslImpl::Init()
 	RegisterSliceOperator(*registry);
 	RegisterConcatOperator(*registry);
 	RegisterUnsqueezeOperator(*registry);
+	RegisterCastOperator(*registry);
 
 	return true;
 }
