Add support for DirectML operators that have constant CPU inputs

#rnx

diff --git a/Engine/Plugins/Experimental/NNE/Source/NNERuntimeRDG/Private/Dml/NNEDmlModel.cpp b/Engine/Plugins/Experimental/NNE/Source/NNERuntimeRDG/Private/Dml/NNEDmlModel.cpp
--- a/Engine/Plugins/Experimental/NNE/Source/NNERuntimeRDG/Private/Dml/NNEDmlModel.cpp
+++ b/Engine/Plugins/Experimental/NNE/Source/NNERuntimeRDG/Private/Dml/NNEDmlModel.cpp
@@ -323,6 +323,7 @@ public:
 		TConstArrayView<int32>		InputIndices;
 		TConstArrayView<int32>		OutputIndices;
 		TConstArrayView<int32>		WeightIndices;
+		TConstArrayView<int32>		ConstantCPUIndices;
 		TConstArrayView<int32>		IntermediateIndices;
 		TConstArrayView<FTensorRDG>	WeightTensors;
 		TConstArrayView<FOpDesc>	Operators;
@@ -546,12 +547,12 @@ private:
 
 		for (int32 Idx = 0; Idx < InGraph.WeightIndices.Num(); ++Idx)
 		{
-			const int32 TensorIdx = InGraph.WeightIndices[Idx];
-			AddInput(TensorIdx);
+			if (InGraph.ConstantCPUIndices.Find(InGraph.WeightIndices[Idx]) == -1)
+			{
+				AddInput(InGraph.WeightIndices[Idx]);
+			}
 		}
 
-		//const int32 OutputIndexStart = InGraph.InputIndices.Num() + InGraph.WeightIndices.Num() + InGraph.IntermediateIndices.Num();
-
 		for (int32 Idx = 0; Idx < InGraph.OutputIndices.Num(); ++Idx)
 		{
 			const int32 TensorIdx = InGraph.OutputIndices[Idx];
@@ -740,6 +741,8 @@ FModel::~FModel()
 //
 bool FModel::Init(TConstArrayView<uint8> ModelData, FDmlDeviceContext* InDevCtx)
 {
+	ConstantCPUTensorIndices.Reset();
+
 	check(ModelData.Num() > 0);
 	FNNERuntimeFormat	Format;
 	int32 GuidSize = FModelInfo::Get()->GetGuidSize();
@@ -770,16 +773,6 @@ bool FModel::Init(TConstArrayView<uint8> ModelData, FDmlDeviceContext* InDevCtx)
 		Tensors.Emplace(NNECore::Internal::FTensor::MakeFromSymbolicDesc(TensorDesc));
 	}
 
-	FGraphBuilder				DmlGraphBuilder;
-	FGraphBuilder::FGraphDesc	DmlGraphDesc;
-
-	DmlGraphDesc.AllTensors = Tensors;
-	DmlGraphDesc.InputIndices = InputTensorIndices;
-	DmlGraphDesc.OutputIndices = OutputTensorIndices;
-	DmlGraphDesc.WeightIndices = WeightTensorIndices;
-	DmlGraphDesc.IntermediateIndices = IntermediateTensorIndices;
-	DmlGraphDesc.WeightTensors = WeightTensorRDGs;
-	
 	TArray<FGraphBuilder::FOpDesc>	DmlGraphOperators;
 	TArray<int32>					OpInputIndices;
 	TArray<int32>					OpOutputIndices;
@@ -790,10 +783,10 @@ bool FModel::Init(TConstArrayView<uint8> ModelData, FDmlDeviceContext* InDevCtx)
 	{
 		const FString TypeName = Format.Operators[Idx].TypeName;
 
-		FGraphBuilder::FOpDesc		OpDesc;
-		TArray<NNECore::Internal::FTensor>				OpInputTensors;
-		TArray<NNECore::Internal::FTensor>				OpOutputTensors;
-		NNECore::FAttributeMap	AttributeMap;
+		FGraphBuilder::FOpDesc					OpDesc;
+		TArray<NNECore::Internal::FTensor>		OpInputTensors;
+		TArray<NNECore::Internal::FTensor>		OpOutputTensors;
+		NNECore::FAttributeMap					AttributeMap;
 
 		OpDesc.InputStart = OpInputIndices.Num();
 		OpDesc.OutputStart = OpOutputIndices.Num();
@@ -839,6 +832,20 @@ bool FModel::Init(TConstArrayView<uint8> ModelData, FDmlDeviceContext* InDevCtx)
 			return false;
 		}
 
+		// Filter out the constant CPU inputs from the graph node inputs
+		TConstArrayView<int32> OpConstantCPUInputs = OpDesc.Op->GetConstantCPUInputs();
+
+		for (int32 InputIdx = OpConstantCPUInputs.Num() - 1; InputIdx >= 0; --InputIdx)
+		{
+			const int32 ConstIdx = OpConstantCPUInputs[InputIdx];
+			const int32 TensorIdx = OpInputIndices[OpDesc.InputStart + ConstIdx];
+
+			OpInputTensors.RemoveAt(ConstIdx);
+			OpInputIndices.RemoveAt(OpDesc.InputStart + ConstIdx);
+
+			ConstantCPUTensorIndices.Emplace(TensorIdx);
+		}
+
 		OpDesc.InputCount = OpInputTensors.Num();
 		OpDesc.OutputCount = OpOutputTensors.Num();
 		OpDesc.DbgName = TypeName;
@@ -846,6 +853,16 @@ bool FModel::Init(TConstArrayView<uint8> ModelData, FDmlDeviceContext* InDevCtx)
 		DmlGraphOperators.Emplace(OpDesc);
 	}
 
+	FGraphBuilder				DmlGraphBuilder;
+	FGraphBuilder::FGraphDesc	DmlGraphDesc;
+
+	DmlGraphDesc.AllTensors = Tensors;
+	DmlGraphDesc.InputIndices = InputTensorIndices;
+	DmlGraphDesc.OutputIndices = OutputTensorIndices;
+	DmlGraphDesc.WeightIndices = WeightTensorIndices;
+	DmlGraphDesc.ConstantCPUIndices = ConstantCPUTensorIndices;
+	DmlGraphDesc.IntermediateIndices = IntermediateTensorIndices;
+	DmlGraphDesc.WeightTensors = WeightTensorRDGs;
 	DmlGraphDesc.Operators = DmlGraphOperators;
 	DmlGraphDesc.OpInputIndices = OpInputIndices;
 	DmlGraphDesc.OpOutputIndices = OpOutputIndices;
@@ -929,7 +946,10 @@ bool FModel::InitCompiledOp(TConstArrayView<int32> OpInputIndices, uint64 Tensor
 
 			for (int32 InputIdx : InputTensorIndices)
 			{
-				Inputs.Emplace(nullptr);
+				if (ConstantCPUTensorIndices.Find(InputIdx) == -1)
+				{
+					Inputs.Emplace(nullptr);
+				}
 			}
 
 			TArray<CD3DX12_RESOURCE_BARRIER, TInlineAllocator<MaxNumInputs>>	Barriers;
@@ -943,8 +963,14 @@ bool FModel::InitCompiledOp(TConstArrayView<int32> OpInputIndices, uint64 Tensor
 				uint8*			UploadBuffPtr = static_cast<uint8*>(RHICmdList.LockBuffer(UploadBuff, 0, TensorDataSize, RLM_WriteOnly_NoOverwrite));
 				uint64			UploadOffset = 0;
 
-				for (const FTensorRDG& Tensor : WeightTensorRDGs)
+				for (int32 WeightIdx : WeightTensorIndices)
 				{
+					if (ConstantCPUTensorIndices.Find(WeightIdx) != -1)
+					{
+						continue;
+					}
+
+					const FTensorRDG&		Tensor = WeightTensorRDGs[WeightIdx];
 					TConstArrayView<uint8>	TensorData = Tensor.GetPreparedData<uint8>();
 
 					FBufferRHIRef WeightBuff;
@@ -1051,8 +1077,10 @@ void FModel::AddDispatchOps_RenderThread(FRDGBuilder& GraphBuilder)
 {
 	const ERDGPassFlags TransitionBuffFlags = ERDGPassFlags::Compute | ERDGPassFlags::NeverCull;
 
-	InputBuffers.SetNumUninitialized(InputTensorIndices.Num() + WeightTensorIndices.Num());
-	OutputBuffers.SetNumUninitialized(OutputTensorIndices.Num());
+	InputBuffers.Reset();
+	OutputBuffers.Reset();
+	InputBuffers.Reserve(InputTensorIndices.Num() + WeightTensorIndices.Num());
+	OutputBuffers.Reserve(OutputTensorIndices.Num());
 
 	for (int32 Idx = 0; Idx < InputTensorIndices.Num(); ++Idx)
 	{
@@ -1063,16 +1091,19 @@ void FModel::AddDispatchOps_RenderThread(FRDGBuilder& GraphBuilder)
 			RDG_EVENT_NAME("FModel_Dispatch_GetInputBuffer"),
 			Params,
 			TransitionBuffFlags,
-			[this, Idx, Params](FRHICommandListImmediate& RHICmdList)
+			[this, Params](FRHICommandListImmediate& RHICmdList)
 			{
-				InputBuffers[Idx] = Params->Buffer->GetRHI();
+				InputBuffers.Emplace(Params->Buffer->GetRHI());
 			}
 		);
 	}
 
 	for (int32 Idx = 0; Idx < WeightTensorIndices.Num(); ++Idx)
 	{
-		InputBuffers[Idx + InputTensorIndices.Num()] = nullptr;
+		if (ConstantCPUTensorIndices.Find(WeightTensorIndices[Idx]) == -1)
+		{
+			InputBuffers.Emplace(nullptr);
+		}
 	}
 
 	for (int32 Idx = 0; Idx < OutputTensorIndices.Num(); ++Idx)
@@ -1084,9 +1115,9 @@ void FModel::AddDispatchOps_RenderThread(FRDGBuilder& GraphBuilder)
 			RDG_EVENT_NAME("FModel_Dispatch_GetOutputBuffer"),
 			Params,
 			TransitionBuffFlags,
-			[this, Idx, Params](FRHICommandListImmediate& RHICmdList)
+			[this, Params](FRHICommandListImmediate& RHICmdList)
 			{
-				OutputBuffers[Idx] = Params->Buffer->GetRHI();
+				OutputBuffers.Emplace(Params->Buffer->GetRHI());
 			}
 		);
 	}
diff --git a/Engine/Plugins/Experimental/NNE/Source/NNERuntimeRDG/Private/Dml/NNEDmlModel.h b/Engine/Plugins/Experimental/NNE/Source/NNERuntimeRDG/Private/Dml/NNEDmlModel.h
--- a/Engine/Plugins/Experimental/NNE/Source/NNERuntimeRDG/Private/Dml/NNEDmlModel.h
+++ b/Engine/Plugins/Experimental/NNE/Source/NNERuntimeRDG/Private/Dml/NNEDmlModel.h
@@ -84,6 +84,8 @@ private:
 	uint32								DescCount;
 	uint32								DescSize;
 
+	TArray<int32>						ConstantCPUTensorIndices;
+
 	FRHIBufferInputArray				InputBuffers;
 	FRHIBufferOutputArray				OutputBuffers;
 #ifdef NNE_USE_D3D12_RESOURCES
diff --git a/Engine/Plugins/Experimental/NNE/Source/NNERuntimeRDG/Private/Dml/NNEDmlOperator.cpp b/Engine/Plugins/Experimental/NNE/Source/NNERuntimeRDG/Private/Dml/NNEDmlOperator.cpp
--- a/Engine/Plugins/Experimental/NNE/Source/NNERuntimeRDG/Private/Dml/NNEDmlOperator.cpp
+++ b/Engine/Plugins/Experimental/NNE/Source/NNERuntimeRDG/Private/Dml/NNEDmlOperator.cpp
@@ -336,6 +336,12 @@ namespace DmlUtil
 //
 // DirectML operator base class
 //
+
+TConstArrayView<int32> FOperatorDml::GetConstantCPUInputs() const
+{
+	return ConstantCPUInputs;
+}
+
 IDMLOperator* FOperatorDml::GetOperator()
 {
 	return DmlOp;
diff --git a/Engine/Plugins/Experimental/NNE/Source/NNERuntimeRDG/Private/Dml/NNEDmlOperator.h b/Engine/Plugins/Experimental/NNE/Source/NNERuntimeRDG/Private/Dml/NNEDmlOperator.h
--- a/Engine/Plugins/Experimental/NNE/Source/NNERuntimeRDG/Private/Dml/NNEDmlOperator.h
+++ b/Engine/Plugins/Experimental/NNE/Source/NNERuntimeRDG/Private/Dml/NNEDmlOperator.h
@@ -93,6 +93,8 @@ public:
 
 	virtual bool Initialize(IDMLDevice* Device, TArrayView<const NNECore::Internal::FTensor> InputTensors, TArrayView<const NNECore::Internal::FTensor> OutputTensors, const NNECore::FAttributeMap& Attributes) = 0;
 
+	virtual TConstArrayView<int32> GetConstantCPUInputs() const;
+
 	IDMLOperator* GetOperator();
 
 protected:
@@ -103,6 +105,7 @@ protected:
 	bool CreateOperator(IDMLDevice* Device, const DML_OPERATOR_DESC& DmlOpDesc);
 
 	TComPtr<IDMLOperator>		DmlOp;
+	DmlUtil::FSmallIntArray		ConstantCPUInputs;
 };
 
 /**
