NNE - Add DML Concat

diff --git a/Engine/Plugins/Experimental/NNE/Source/NNEQA/Resources/NNEQATestDesc.json b/Engine/Plugins/Experimental/NNE/Source/NNEQA/Resources/NNEQATestDesc.json
--- a/Engine/Plugins/Experimental/NNE/Source/NNEQA/Resources/NNEQATestDesc.json
+++ b/Engine/Plugins/Experimental/NNE/Source/NNEQA/Resources/NNEQATestDesc.json
@@ -348,6 +348,15 @@
 			"targets": [
 				{ "target": "Shape" }
 			]
+		},
+		{
+			"category": "Concat",
+			"runtimes": [
+				{ "name": "NNERuntimeRDGHlsl", "skip": true }
+			],
+			"targets": [
+				{ "target": "Concat" }
+			]
 		}
 	],
 
@@ -823,6 +832,18 @@
 				{ "inputs":  [ { "shape": [ 1, 128, 50, 55 ] } ], "outputs": [ { "shape": [ 4 ], "type": "int64"} ] },
 				{ "inputs":  [ { "shape": [ 2, 1, 3, 1, 2 ] } ],  "outputs": [ { "shape": [ 5 ], "type": "int64"} ] }
 			]
+		},
+		{ 
+			"name": "Concat.Concat1",
+			"datasets": [
+				{ "inputs": [ { "shape": [ 2, 3 ] }, { "shape": [ 2, 4 ] } ], "outputs": [ { "shape": [ 2, 7 ] } ] }
+			]
+		},
+		{ 
+			"name": "Concat.Concat2",
+			"datasets": [
+				{ "inputs": [ { "shape": [ 4, 2, 4 ] }, { "shape": [ 6, 2, 4 ] }, { "shape": [ 1, 2, 4 ] } ], "outputs": [ { "shape": [ 11, 2, 4 ] } ] }
+			]
 		}
 	],
 
@@ -1195,6 +1216,22 @@
 			"attribute_maps": [
 				{ "attributes": [] }
 			]
+		},
+		{
+			"name": "Concat.Concat1",
+			"attribute_maps": [
+				{ "attributes": [
+					{ "name": "axis", "value": 1, "type": "Int32" }
+				] }
+			]
+		},
+		{
+			"name": "Concat.Concat2",
+			"attribute_maps": [
+				{ "attributes": [
+					{ "name": "axis", "value": -3, "type": "Int32" }
+				] }
+			]
 		}
 	]
 
diff --git a/Engine/Plugins/Experimental/NNE/Source/NNERuntimeRDG/Private/Dml/NNEDmlOperatorConcat.cpp b/Engine/Plugins/Experimental/NNE/Source/NNERuntimeRDG/Private/Dml/NNEDmlOperatorConcat.cpp
new file mode 100644
--- /dev/null
+++ b/Engine/Plugins/Experimental/NNE/Source/NNERuntimeRDG/Private/Dml/NNEDmlOperatorConcat.cpp
@@ -0,0 +1,93 @@
+// Copyright Epic Games, Inc. All Rights Reserved.
+
+#ifdef NNE_USE_DIRECTML
+
+#include "NNEDmlOperator.h"
+#include "Misc/EnumerateRange.h"
+#include "Algo/Find.h"
+
+namespace UE::NNERuntimeRDG::Private::Dml
+{
+
+/**
+ * Concat
+ */
+class FOperatorDmlConcat : public FOperatorDml
+{
+	
+
+public:
+
+	static FOperatorDml* Create()
+	{
+		return new FOperatorDmlConcat();
+	}
+
+	//
+	//
+	//
+	virtual bool Initialize(IDMLDevice* Device, TArrayView<const NNECore::Internal::FTensor> InputTensors, TArrayView<const NNECore::Internal::FTensor> OutputTensors, const NNECore::FAttributeMap& Attributes) override
+	{
+		check(InputTensors.Num() > 0);
+		check(OutputTensors.Num() == 1);
+
+		int32 InputDims = InputTensors[0].GetShape().Rank();
+		auto NormalizeAxis = [InputDims] (int32 Axis) -> uint32
+		{
+			if(Axis < 0)
+			{
+				Axis += InputDims;
+			}
+			return (uint32) Axis;
+		};
+		uint32 Axis = NormalizeAxis(Attributes.GetValue<int32>(TEXT("axis")));
+		check(Axis < (uint32) InputDims);
+
+		TArray<DmlUtil::FTensorDesc> InputDescs;
+		InputDescs.SetNum(InputTensors.Num());
+		TArray<DML_TENSOR_DESC> DmlInputDescs;
+		DmlInputDescs.SetNumUninitialized(InputTensors.Num());
+		
+		for(TConstEnumerateRef<NNECore::Internal::FTensor> It : EnumerateRange(InputTensors))
+		{
+			const NNECore::Internal::FTensor& Tensor = *It;
+			int Idx = It.GetIndex();
+			if(Algo::Find(Tensor.GetShape().GetData(), 0) == nullptr)
+			{
+				DmlUtil::FTensorDesc& InputTensorDesc = InputDescs[Idx];
+				if (!InputTensorDesc.InitFromTensor(Tensor, Tensor.GetShape().Rank()))
+				{
+					UE_LOG(LogNNE, Error, TEXT("Failed to initialize Concat input for DML inference"));
+					return false;
+				}
+				DmlInputDescs[Idx] = InputTensorDesc.Desc;
+			}
+		}
+
+		check(InputDescs.Num() > 0);
+
+		DmlUtil::FTensorDesc OutputTensorDesc;
+		if (!OutputTensorDesc.InitFromTensor(OutputTensors[0], OutputTensors[0].GetShape().Rank()))
+		{
+			UE_LOG(LogNNE, Error, TEXT("Failed to initialize Concat output for DML inference"));
+			return false;
+		}
+
+
+		DML_JOIN_OPERATOR_DESC DmlJoinOpDesc = {};
+		DmlJoinOpDesc.InputCount = (uint32) DmlInputDescs.Num();
+		DmlJoinOpDesc.InputTensors = DmlInputDescs.GetData();
+		DmlJoinOpDesc.OutputTensor = &OutputTensorDesc.Desc;
+		DmlJoinOpDesc.Axis = Axis;
+
+		return CreateOperator(Device, DML_OPERATOR_DESC{ DML_OPERATOR_JOIN, &DmlJoinOpDesc} );
+
+	}
+};
+
+// Register Concat operator on Module startup
+NNE_DML_REGISTER_OP(Concat)
+
+} // namespace UE::NNERuntimeRDG::Private::Dml
+
+#endif // NNE_USE_DIRECTML
