NNE : Add HLSL Gather operator for shape inference

Test:
* Added unit test for gather operator shape inference helper.

#rnx

diff --git a/Engine/Plugins/Experimental/NNE/Source/NNEQA/Private/NNEQAUnitTestOpGather.cpp b/Engine/Plugins/Experimental/NNE/Source/NNEQA/Private/NNEQAUnitTestOpGather.cpp
new file mode 100644
--- /dev/null
+++ b/Engine/Plugins/Experimental/NNE/Source/NNEQA/Private/NNEQAUnitTestOpGather.cpp
@@ -0,0 +1,171 @@
+// Copyright Epic Games, Inc. All Rights Reserved.
+
+#include "CoreMinimal.h"
+#include "NNECoreTensor.h"
+#include "NNECoreTypes.h"
+#include "NNEQAUnitTestHelper.h"
+#include "NNERuntimeRDGHelperGather.h"
+
+#if WITH_DEV_AUTOMATION_TESTS
+#include "Misc/AutomationTest.h"
+#endif
+
+namespace UE::NNEQA::Private::NNERuntimeRDG::GatherOp
+{
+
+#if WITH_DEV_AUTOMATION_TESTS
+
+	using namespace NNECore::Internal;
+	using namespace UE::NNERuntimeRDG::Internal;
+
+	IMPLEMENT_NNE_SHAPEINFERENCEHELPER_UNIT_AUTOMATION_TEST(FGatherCPUHelperConstOuput, "System.Engine.MachineLearning.NNE.UnitTest.GatherHelper.ConstOutput");
+	bool FGatherCPUHelperConstOuput::RunTest(const FString& Parameter)
+	{
+		FTensor XC1 = MakeConstTensorInt64(TEXT("XC1"), { 1 }, { 0 });
+		FTensor XC20 = MakeConstTensorInt64(TEXT("XC20"), { 20 }, { 0,1,2,3,4,5,6,7,8,9,0,1,2,3,4,5,6,7,8,9 });
+		FTensor X1 = MakeTensor(TEXT("X"), { 0 }, ENNETensorDataType::Int64);
+
+		FTensor Y = MakeTensor(TEXT("Y"), { 1 }, ENNETensorDataType::Int64);
+		CPUHelper::Gather::Apply(XC1 , XC1, 0, Y);
+		UTEST_TRUE(TEXT("Y const if inputs are const"), Y.HasPreparedData());
+
+		Y = MakeTensor(TEXT("Y"), { 1 }, ENNETensorDataType::Int64);
+		CPUHelper::Gather::Apply(XC1, X1, 0, Y);
+		UTEST_FALSE(TEXT("Y not const if data not const"), Y.HasPreparedData());
+
+		Y = MakeTensor(TEXT("Y"), { 1 }, ENNETensorDataType::Int64);
+		CPUHelper::Gather::Apply(X1, XC1, 0, Y);
+		UTEST_FALSE(TEXT("Y not const if indices not const"), Y.HasPreparedData());
+
+		Y = MakeTensor(TEXT("Y"), { 20 }, ENNETensorDataType::Int64);
+		CPUHelper::Gather::Apply(XC20, XC1, 0, Y);
+		UTEST_FALSE(TEXT("Y not const if input is too large "), Y.HasPreparedData());
+
+		return true;
+	}
+
+	IMPLEMENT_NNE_SHAPEINFERENCEHELPER_UNIT_AUTOMATION_TEST(FGatherCPUHelperRank1Int64Indices, "System.Engine.MachineLearning.NNE.UnitTest.GatherHelper.Rank1Int64Indices");
+	bool FGatherCPUHelperRank1Int64Indices::RunTest(const FString& Parameter)
+	{
+		FTensor XC4i64 = MakeConstTensorInt64(TEXT("XC4"), { 4 }, {    0,    2,    -3,     5 });
+		FTensor XC4i32 = MakeConstTensorInt32(TEXT("XC4"), { 4 }, {   -1,    3,    13,     8 });
+		FTensor XC4f32 = MakeConstTensor(TEXT("XC4"),      { 4 }, { 0.1f, 2.5f, -3.1f, 15.0f });
+
+		FTensor XC3 = MakeConstTensorInt64(TEXT("XC3"), { 3 }, { 0, 2 ,3 });
+		FTensor XC2 = MakeConstTensorInt64(TEXT("XC2"), { 2 }, { 1, 2 });
+
+		//int64 Data
+		{
+			FTensor Y = MakeTensor(TEXT("Y"), { 3 }, ENNETensorDataType::Int64);
+			CPUHelper::Gather::Apply(XC4i64, XC3, 0, Y);
+			UTEST_TRUE(TEXT("Y const if input is const"), Y.HasPreparedData());
+			UTEST_EQUAL(TEXT("Gather(XC4i64,XC3)[0]"), Y.GetPreparedData<int64>()[0], (int64)0);
+			UTEST_EQUAL(TEXT("Gather(XC4i64,XC3)[1]"), Y.GetPreparedData<int64>()[1], (int64)-3);
+			UTEST_EQUAL(TEXT("Gather(XC4i64,XC3)[2]"), Y.GetPreparedData<int64>()[2], (int64)5);
+
+			Y = MakeTensor(TEXT("Y"), { 2 }, ENNETensorDataType::Int64);
+			CPUHelper::Gather::Apply(XC4i64, XC2, 0, Y);
+			UTEST_TRUE(TEXT("Y const if input is const"), Y.HasPreparedData());
+			UTEST_EQUAL(TEXT("Gather(XC4i64,XC2)[0]"), Y.GetPreparedData<int64>()[0], (int64)2);
+			UTEST_EQUAL(TEXT("Gather(XC4i64,XC2)[1]"), Y.GetPreparedData<int64>()[1], (int64)-3);
+		}
+
+		//int32 Data
+		{
+			FTensor Y = MakeTensor(TEXT("Y"), { 3 }, ENNETensorDataType::Int32);
+			CPUHelper::Gather::Apply(XC4i32, XC3, 0, Y);
+			UTEST_TRUE(TEXT("Y const if input is const"), Y.HasPreparedData());
+			UTEST_EQUAL(TEXT("Gather(XC4i32,XC3)[0]"), Y.GetPreparedData<int32>()[0], -1);
+			UTEST_EQUAL(TEXT("Gather(XC4i32,XC3)[1]"), Y.GetPreparedData<int32>()[1], 13);
+			UTEST_EQUAL(TEXT("Gather(XC4i32,XC3)[2]"), Y.GetPreparedData<int32>()[2], 8);
+
+			Y = MakeTensor(TEXT("Y"), { 2 }, ENNETensorDataType::Int32);
+			CPUHelper::Gather::Apply(XC4i32, XC2, 0, Y);
+			UTEST_TRUE(TEXT("Y const if input is const"), Y.HasPreparedData());
+			UTEST_EQUAL(TEXT("Gather(XC4i32,XC2)[0]"), Y.GetPreparedData<int32>()[0], 3);
+			UTEST_EQUAL(TEXT("Gather(XC4i32,XC2)[1]"), Y.GetPreparedData<int32>()[1], 13);
+		}
+
+		//float Data
+		{
+			FTensor Y = MakeTensor(TEXT("Y"), { 3 }, ENNETensorDataType::Float);
+			CPUHelper::Gather::Apply(XC4f32, XC3, 0, Y);
+			UTEST_TRUE(TEXT("Y const if input is const"), Y.HasPreparedData());
+			UTEST_EQUAL(TEXT("Gather(XC4f32,XC3)[0]"), Y.GetPreparedData<float>()[0], 0.1f);
+			UTEST_EQUAL(TEXT("Gather(XC4f32,XC3)[1]"), Y.GetPreparedData<float>()[1], -3.1f);
+			UTEST_EQUAL(TEXT("Gather(XC4f32,XC3)[2]"), Y.GetPreparedData<float>()[2], 15.0f);
+
+			Y = MakeTensor(TEXT("Y"), { 2 }, ENNETensorDataType::Float);
+			CPUHelper::Gather::Apply(XC4f32, XC2, 0, Y);
+			UTEST_TRUE(TEXT("Y const if input is const"), Y.HasPreparedData());
+			UTEST_EQUAL(TEXT("Gather(XC4f32,XC2)[0]"), Y.GetPreparedData<float>()[0], 2.5f);
+			UTEST_EQUAL(TEXT("Gather(XC4f32,XC2)[1]"), Y.GetPreparedData<float>()[1], -3.1f);
+		}
+
+		return true;
+	}
+
+	IMPLEMENT_NNE_SHAPEINFERENCEHELPER_UNIT_AUTOMATION_TEST(FGatherCPUHelperRank1Int32Indices, "System.Engine.MachineLearning.NNE.UnitTest.GatherHelper.Rank1Int32Indices");
+	bool FGatherCPUHelperRank1Int32Indices::RunTest(const FString& Parameter)
+	{
+		FTensor XC4i64 = MakeConstTensorInt64(TEXT("XC4"), { 4 }, {    0,    2,    -3,     5 });
+		FTensor XC4i32 = MakeConstTensorInt32(TEXT("XC4"), { 4 }, {   -1,    3,    13,     8 });
+		FTensor XC4f32 = MakeConstTensor(TEXT("XC4"),      { 4 }, { 0.1f, 2.5f, -3.1f, 15.0f });
+
+		FTensor XC3 = MakeConstTensorInt32(TEXT("XC3"), { 3 }, { 0, 2 ,3 });
+		FTensor XC2 = MakeConstTensorInt32(TEXT("XC2"), { 2 }, { 1, 2 });
+
+		//int64 Data
+		{
+			FTensor Y = MakeTensor(TEXT("Y"), { 3 }, ENNETensorDataType::Int64);
+			CPUHelper::Gather::Apply(XC4i64, XC3, 0, Y);
+			UTEST_TRUE(TEXT("Y const if input is const"), Y.HasPreparedData());
+			UTEST_EQUAL(TEXT("Gather(XC4i64,XC3)[0]"), Y.GetPreparedData<int64>()[0], (int64)0);
+			UTEST_EQUAL(TEXT("Gather(XC4i64,XC3)[1]"), Y.GetPreparedData<int64>()[1], (int64)-3);
+			UTEST_EQUAL(TEXT("Gather(XC4i64,XC3)[2]"), Y.GetPreparedData<int64>()[2], (int64)5);
+
+			Y = MakeTensor(TEXT("Y"), { 2 }, ENNETensorDataType::Int64);
+			CPUHelper::Gather::Apply(XC4i64, XC2, 0, Y);
+			UTEST_TRUE(TEXT("Y const if input is const"), Y.HasPreparedData());
+			UTEST_EQUAL(TEXT("Gather(XC4i64,XC2)[0]"), Y.GetPreparedData<int64>()[0], (int64)2);
+			UTEST_EQUAL(TEXT("Gather(XC4i64,XC2)[1]"), Y.GetPreparedData<int64>()[1], (int64)-3);
+		}
+
+		//int32 Data
+		{
+			FTensor Y = MakeTensor(TEXT("Y"), { 3 }, ENNETensorDataType::Int32);
+			CPUHelper::Gather::Apply(XC4i32, XC3, 0, Y);
+			UTEST_TRUE(TEXT("Y const if input is const"), Y.HasPreparedData());
+			UTEST_EQUAL(TEXT("Gather(XC4i32,XC3)[0]"), Y.GetPreparedData<int32>()[0], -1);
+			UTEST_EQUAL(TEXT("Gather(XC4i32,XC3)[1]"), Y.GetPreparedData<int32>()[1], 13);
+			UTEST_EQUAL(TEXT("Gather(XC4i32,XC3)[2]"), Y.GetPreparedData<int32>()[2], 8);
+
+			Y = MakeTensor(TEXT("Y"), { 2 }, ENNETensorDataType::Int32);
+			CPUHelper::Gather::Apply(XC4i32, XC2, 0, Y);
+			UTEST_TRUE(TEXT("Y const if input is const"), Y.HasPreparedData());
+			UTEST_EQUAL(TEXT("Gather(XC4i32,XC2)[0]"), Y.GetPreparedData<int32>()[0], 3);
+			UTEST_EQUAL(TEXT("Gather(XC4i32,XC2)[1]"), Y.GetPreparedData<int32>()[1], 13);
+		}
+
+		//float Data
+		{
+			FTensor Y = MakeTensor(TEXT("Y"), { 3 }, ENNETensorDataType::Float);
+			CPUHelper::Gather::Apply(XC4f32, XC3, 0, Y);
+			UTEST_TRUE(TEXT("Y const if input is const"), Y.HasPreparedData());
+			UTEST_EQUAL(TEXT("Gather(XC4f32,XC3)[0]"), Y.GetPreparedData<float>()[0], 0.1f);
+			UTEST_EQUAL(TEXT("Gather(XC4f32,XC3)[1]"), Y.GetPreparedData<float>()[1], -3.1f);
+			UTEST_EQUAL(TEXT("Gather(XC4f32,XC3)[2]"), Y.GetPreparedData<float>()[2], 15.0f);
+
+			Y = MakeTensor(TEXT("Y"), { 2 }, ENNETensorDataType::Float);
+			CPUHelper::Gather::Apply(XC4f32, XC2, 0, Y);
+			UTEST_TRUE(TEXT("Y const if input is const"), Y.HasPreparedData());
+			UTEST_EQUAL(TEXT("Gather(XC4f32,XC2)[0]"), Y.GetPreparedData<float>()[0], 2.5f);
+			UTEST_EQUAL(TEXT("Gather(XC4f32,XC2)[1]"), Y.GetPreparedData<float>()[1], -3.1f);
+		}
+
+		return true;
+	}
+
+#endif //WITH_DEV_AUTOMATION_TESTS
+
+} // namespace UE::NNEQA::Private::NNERuntimeRDG::GatherOp
diff --git a/Engine/Plugins/Experimental/NNE/Source/NNERuntimeRDG/Internal/NNERuntimeRDGHelperGather.h b/Engine/Plugins/Experimental/NNE/Source/NNERuntimeRDG/Internal/NNERuntimeRDGHelperGather.h
new file mode 100644
--- /dev/null
+++ b/Engine/Plugins/Experimental/NNE/Source/NNERuntimeRDG/Internal/NNERuntimeRDGHelperGather.h
@@ -0,0 +1,12 @@
+// Copyright Epic Games, Inc. All Rights Reserved.
+
+#pragma once 
+
+#include "CoreMinimal.h"
+
+namespace UE::NNECore::Internal { class FTensor; }
+
+namespace UE::NNERuntimeRDG::Internal::CPUHelper::Gather
+{
+	void NNERUNTIMERDG_API Apply(const NNECore::Internal::FTensor& DataTensor, const NNECore::Internal::FTensor& IndicesTensor, int32 Axis, NNECore::Internal::FTensor& OutputTensor);
+} // UE::NNERuntimeRDG::Internal::CPUHelper::Gather
\ No newline at end of file
diff --git a/Engine/Plugins/Experimental/NNE/Source/NNERuntimeRDG/Private/Helper/NNERuntimeRDGHelperGather.cpp b/Engine/Plugins/Experimental/NNE/Source/NNERuntimeRDG/Private/Helper/NNERuntimeRDGHelperGather.cpp
new file mode 100644
--- /dev/null
+++ b/Engine/Plugins/Experimental/NNE/Source/NNERuntimeRDG/Private/Helper/NNERuntimeRDGHelperGather.cpp
@@ -0,0 +1,87 @@
+// Copyright Epic Games, Inc. All Rights Reserved.
+
+#include "NNERuntimeRDGHelperGather.h"
+#include "NNERuntimeRDGTensorIdxIterator.h"
+
+namespace UE::NNERuntimeRDG::Internal::CPUHelper::Gather
+{
+	template <typename TData, typename TInd> void Gather1DResolvedTypes(const NNECore::Internal::FTensor& DataTensor, const NNECore::Internal::FTensor& IndicesTensor, NNECore::Internal::FTensor& OutputTensor)
+	{
+		TConstArrayView<TData> DataData = DataTensor.GetPreparedData<TData>();
+		TConstArrayView<TInd> IndicesData = IndicesTensor.GetPreparedData<TInd>();
+		TArray<TData> OutputData;
+
+		check(IndicesData.Num() == OutputTensor.GetShape().GetData()[0]);
+		OutputData.Reserve(IndicesData.Num());
+
+		for (TInd indices : IndicesData)
+		{
+			OutputData.Add(DataData[indices]);
+		}
+
+		OutputTensor.SetPreparedData<TData>(OutputData);
+	}
+
+	template <typename TInd> void Gather1DResolvedIndices(const NNECore::Internal::FTensor& DataTensor, const NNECore::Internal::FTensor& IndicesTensor, NNECore::Internal::FTensor& OutputTensor)
+	{
+		switch (DataTensor.GetDataType())
+		{
+		case ENNETensorDataType::Float:
+			Gather1DResolvedTypes<float, TInd>(DataTensor, IndicesTensor, OutputTensor);
+			break;
+		case ENNETensorDataType::Int32:
+			Gather1DResolvedTypes<int32, TInd>(DataTensor, IndicesTensor, OutputTensor);
+			break;
+		case ENNETensorDataType::Int64:
+			Gather1DResolvedTypes<int64, TInd>(DataTensor, IndicesTensor, OutputTensor);
+			break;
+		}
+	}
+
+	void Apply(const NNECore::Internal::FTensor& DataTensor, const NNECore::Internal::FTensor& IndicesTensor, int32 Axis, NNECore::Internal::FTensor& OutputTensor)
+	{
+		static constexpr int32 MaxItemInOutputTensor = NNECore::FTensorShape::MaxRank * 2;
+
+		if (OutputTensor.GetVolume() >= MaxItemInOutputTensor)
+		{
+			return;
+		}
+
+		if ((DataTensor.GetDataType() != ENNETensorDataType::Float) &&
+			(DataTensor.GetDataType() != ENNETensorDataType::Int64) &&
+			(DataTensor.GetDataType() != ENNETensorDataType::Int32))
+		{
+			return;
+		}
+		check(OutputTensor.GetDataType() == DataTensor.GetDataType());
+
+		if ((IndicesTensor.GetDataType() != ENNETensorDataType::Int32) &&
+			(IndicesTensor.GetDataType() != ENNETensorDataType::Int64))
+		{
+			return;
+		}
+
+		if (!DataTensor.HasPreparedData() || !IndicesTensor.HasPreparedData())
+		{
+			return;
+		}
+
+		if (DataTensor.GetShape().Rank() != 1 || IndicesTensor.GetShape().Rank() != 1)
+		{
+			return;
+		}
+
+		check(Axis == 0);
+
+		if (IndicesTensor.GetDataType() == ENNETensorDataType::Int64)
+		{
+			Gather1DResolvedIndices<int64>(DataTensor, IndicesTensor, OutputTensor);
+		}
+		else
+		{
+			Gather1DResolvedIndices<int32>(DataTensor, IndicesTensor, OutputTensor);
+		}
+		
+	}
+	
+} // UE::NNERuntimeRDG::Private::CPUHelper::Gather
diff --git a/Engine/Plugins/Experimental/NNE/Source/NNERuntimeRDG/Private/Hlsl/NNERuntimeRDGGather.cpp b/Engine/Plugins/Experimental/NNE/Source/NNERuntimeRDG/Private/Hlsl/NNERuntimeRDGGather.cpp
--- a/Engine/Plugins/Experimental/NNE/Source/NNERuntimeRDG/Private/Hlsl/NNERuntimeRDGGather.cpp
+++ b/Engine/Plugins/Experimental/NNE/Source/NNERuntimeRDG/Private/Hlsl/NNERuntimeRDGGather.cpp
@@ -1,6 +1,7 @@
 // Copyright Epic Games, Inc. All Rights Reserved.
 
 #include "NNERuntimeRDGGather.h"
+#include "NNERuntimeRDGHelperGather.h"
 #include "NNERuntimeRDGHlslHelper.h"
 #include "NNEHlslShadersGatherCS.h"
 #include "NNECoreAttributeMap.h"
@@ -33,7 +34,44 @@ namespace UE::NNERuntimeRDG::Private::Hlsl
 			check(InputTensors.Num() == 2)
 			check(OutputTensors.Num() == 1)
 			UE_LOG(LogNNE, Warning, TEXT("Gather shape inference is not implemented at the moment"));
-			return -1;
+
+			const NNECore::Internal::FTensor& Data = *InputTensors[0];
+			const NNECore::Internal::FTensor& Indices = *InputTensors[1];
+			const NNECore::FTensorShape& DataShape = Data.GetShape();
+			const NNECore::FTensorShape& IndicesShape = Indices.GetShape();
+
+			const int32 OutputRank = IndicesShape.Rank() + DataShape.Rank() - 1;
+			TArray<uint32> OutputShapeData;
+			int32 DataRankIdx = 0;
+
+			OutputShapeData.Reserve(OutputRank);
+			for (;DataRankIdx < Axis; ++DataRankIdx)
+			{
+				OutputShapeData.Add(DataShape.GetData()[DataRankIdx]);
+			}
+
+			OutputShapeData.Append(IndicesShape.GetData());
+			++DataRankIdx;
+
+			for (; DataRankIdx < DataShape.Rank(); ++DataRankIdx)
+			{
+				OutputShapeData.Add(DataShape.GetData()[DataRankIdx]);
+			}
+
+			NNECore::FTensorShape OutputShape = NNECore::FTensorShape::Make(OutputShapeData);
+
+			check(OutputShape.Rank() == OutputRank);
+			OutputTensors[0]->SetShape(OutputShape);
+
+			Internal::CPUHelper::Gather::Apply(Data, Indices, Axis, *OutputTensors[0]);
+
+			if (!OutputTensors[0]->HasPreparedData())
+			{
+				UE_LOG(LogNNE, Warning, TEXT("Gather: Output could not be computed as a constant tensor, however Gather is not implemented on GPU at the moment."));
+				return -1;
+			}
+
+			return 0;
 		};
 		
 		virtual bool Initialize(TConstArrayView<NNECore::FTensorDesc> InputTensorDescs, TConstArrayView<NNECore::FTensorDesc> OutputTensorDescs, const NNECore::FAttributeMap& Attributes) override
@@ -86,6 +124,8 @@ namespace UE::NNERuntimeRDG::Private::Hlsl
 
 		virtual void Dispatch(FRDGBuilder& GraphBuilder, TConstArrayView<FTensorRDGRef> InputTensors, TConstArrayView<FTensorRDGRef> OutputTensors) override
 		{
+			UE_LOG(LogNNE, Warning, TEXT("Gather: Output should be constant and already uploaded to GPU memory. Dispatch should not need to be called."));
+			/*
 			using namespace UE::NNEHlslShaders::Internal;
 
 			check(InputTensors.Num() == 2)
@@ -125,6 +165,7 @@ namespace UE::NNERuntimeRDG::Private::Hlsl
 				ComputeShader,
 				Parameters,
 				ThreadGroupCount);
+			*/
 		}
 	};
 
@@ -132,14 +173,21 @@ namespace UE::NNERuntimeRDG::Private::Hlsl
 	{
 		bool bIsValid = true;
 
+		//This match version 13 of the Gather operator
+		//[REDACTED_PROFILE_URL]/onnx/blob/main/docs/Operators.md#Gather
 		FAttributeValidator AttributeValidator;
 		AttributeValidator.AddOptional(TEXT("axis"), ENNEAttributeDataType::Int32);
 		bIsValid &= AttributeValidator.Validate(AttributeMap);
 
 		FInputValidator InputValidator;
-		InputValidator.AddSupportedType(ENNETensorDataType::Float);
-		InputValidator.AddRequired();
-		InputValidator.AddRequired();//Indices should be int32 or int64
+		InputValidator.SetTemplateCount(2);
+	
+		InputValidator.AddSupportedType(ENNETensorDataType::Float, 0);
+		InputValidator.AddRequired(0);
+
+		InputValidator.AddSupportedType(ENNETensorDataType::Int32, 1);
+		InputValidator.AddSupportedType(ENNETensorDataType::Int64, 1);
+		InputValidator.AddRequired(1);
 		bIsValid &= InputValidator.Validate(InputTypes);
 
 		return bIsValid;
diff --git a/Engine/Plugins/Experimental/NNE/Source/NNERuntimeRDG/Private/NNERuntimeRDGHlsl.cpp b/Engine/Plugins/Experimental/NNE/Source/NNERuntimeRDG/Private/NNERuntimeRDGHlsl.cpp
--- a/Engine/Plugins/Experimental/NNE/Source/NNERuntimeRDG/Private/NNERuntimeRDGHlsl.cpp
+++ b/Engine/Plugins/Experimental/NNE/Source/NNERuntimeRDG/Private/NNERuntimeRDGHlsl.cpp
@@ -15,6 +15,7 @@
 #include "Hlsl/NNERuntimeRDGElementWiseBinary.h"
 #include "Hlsl/NNERuntimeRDGElementWiseUnary.h"
 #include "Hlsl/NNERuntimeRDGElementWiseVariadic.h"
+#include "Hlsl/NNERuntimeRDGGather.h"
 #include "Hlsl/NNERuntimeRDGGemm.h"
 #include "Hlsl/NNERuntimeRDGInstanceNormalization.h"
 #include "Hlsl/NNERuntimeRDGPad.h"
@@ -49,6 +50,7 @@ bool UNNERuntimeRDGHlslImpl::Init()
 	RegisterConcatOperator(*registry);
 	RegisterUnsqueezeOperator(*registry);
 	RegisterCastOperator(*registry);
+	RegisterGatherOperator(*registry);
 
 	return true;
 }
