NNE - Add DirectML Slice
#rnx

diff --git a/Engine/Plugins/Experimental/NNE/Source/NNEQA/Resources/NNEQATestDesc.json b/Engine/Plugins/Experimental/NNE/Source/NNEQA/Resources/NNEQATestDesc.json
--- a/Engine/Plugins/Experimental/NNE/Source/NNEQA/Resources/NNEQATestDesc.json
+++ b/Engine/Plugins/Experimental/NNE/Source/NNEQA/Resources/NNEQATestDesc.json
@@ -113,6 +113,24 @@
 					]
 				}
 			]
+		},
+		{
+			"category": "SingleOp",
+			"is_model_test": true,
+			"targets": [
+				{
+					"target": "Slice",
+					"runtimes": [
+						{ "name": "NNERuntimeRDGHlsl", "skip": true, "comment": "Not implemented" }
+					]
+				},
+				{
+					"target": "SliceNoSteps",
+					"runtimes": [
+						{ "name": "NNERuntimeRDGHlsl", "skip": true, "comment": "Not implemented" }
+					]
+				}
+			]
 		}
 	],
 
diff --git a/Engine/Plugins/Experimental/NNE/Source/NNERuntimeRDG/Private/Dml/NNEDmlOperatorSlice.cpp b/Engine/Plugins/Experimental/NNE/Source/NNERuntimeRDG/Private/Dml/NNEDmlOperatorSlice.cpp
new file mode 100644
--- /dev/null
+++ b/Engine/Plugins/Experimental/NNE/Source/NNERuntimeRDG/Private/Dml/NNEDmlOperatorSlice.cpp
@@ -0,0 +1,165 @@
+// Copyright Epic Games, Inc. All Rights Reserved.
+
+#ifdef NNE_USE_DIRECTML
+
+#include "NNEDmlOperator.h"
+#include "Misc/EnumerateRange.h"
+#include "Algo/Transform.h"
+#include "Algo/ForEach.h"
+
+#include <numeric>
+
+namespace UE::NNERuntimeRDG::Private::Dml
+{
+
+/**
+ * Slice
+ */
+class FOperatorDmlSlice : public FOperatorDml
+{
+	
+
+public:
+
+	static FOperatorDml* Create()
+	{
+		return new FOperatorDmlSlice();
+	}
+
+	//
+	//
+	//
+	virtual bool Initialize(IDMLDevice* Device, TArrayView<const NNECore::Internal::FTensor> InputTensors, TArrayView<const NNECore::Internal::FTensor> OutputTensors, const NNECore::FAttributeMap& Attributes) override
+	{
+		check(InputTensors.Num() >= 3);
+		check(InputTensors.Num() <= 5);
+		check(OutputTensors.Num() == 1);
+
+		for(int Idx = 1; Idx < InputTensors.Num(); Idx++)
+		{
+			check(InputTensors[Idx].GetShape().Rank() == 1);
+			check(InputTensors[Idx].GetDataType() == ENNETensorDataType::Int32 
+			   || InputTensors[Idx].GetDataType() == ENNETensorDataType::Int64);
+			check(InputTensors[Idx].HasPreparedData());
+			ConstantCPUInputs.Add(Idx);
+		}
+
+		// Initialize Input tensor desc
+        DmlUtil::FTensorDesc DmlInputTensorDesc;
+        if (!DmlInputTensorDesc.InitFromTensor(InputTensors[0], InputTensors[0].GetShape().Rank()))
+        {
+            UE_LOG(LogNNE, Error, TEXT("Failed to initialize Slice input for DML inference"));
+            return false;
+        }
+
+		// Initialize Output tensor desc
+        DmlUtil::FTensorDesc DmlOutputTensorDesc;
+        if (!DmlOutputTensorDesc.InitFromTensor(OutputTensors[0], OutputTensors[0].GetShape().Rank()))
+        {
+            UE_LOG(LogNNE, Error, TEXT("Failed to initialize Slice Output for DML inference"));
+            return false;
+        }
+		TConstArrayView<int32> Starts = InputTensors[1].GetPreparedData<int32>();
+		TConstArrayView<int32> Ends = InputTensors[2].GetPreparedData<int32>();
+		check(Starts.Num() == Ends.Num());
+
+		DmlUtil::FSmallIntArray Axes;
+		DmlUtil::FSmallIntArray Steps;
+
+		if(InputTensors.Num() >= 4)
+		{
+			auto NormalizeAxes = [NumDims = InputTensors[0].GetShape().Rank()] (int32& Axis)
+			{
+				if(Axis < 0)
+				{
+					Axis += NumDims;
+				}
+			};
+			Axes = InputTensors[3].GetPreparedData<int32>();
+			check(Axes.Num() == Starts.Num());
+			Algo::ForEach(Axes, NormalizeAxes);
+		}
+		else
+		{
+			Axes.SetNumUninitialized(Starts.Num());
+			std::iota(Axes.begin(), Axes.end(), 0);
+		}
+
+		if(InputTensors.Num() >= 5)
+		{
+			Steps = InputTensors[4].GetPreparedData<int32>();
+			check(Steps.Num() == Axes.Num());
+		}
+		else
+		{
+			Steps.Init(1, Axes.Num());
+		}
+
+		DmlUtil::FSmallUIntArray OutputShape, Sizes, Offsets;
+		Algo::Transform(InputTensors[0].GetShape().GetData(), OutputShape, [](int32 In) { return (uint32) In; });
+		Sizes = OutputShape;
+		Offsets.SetNumZeroed(OutputShape.Num());
+
+		DmlUtil::FSmallIntArray Strides;
+		Strides.Init(1, OutputShape.Num());
+
+		for (TConstEnumerateRef<int32> Elem : EnumerateRange(Starts))
+		{
+			int32 Idx = Elem.GetIndex();
+			int32 Start = *Elem;
+			int32 End = Ends[Idx];
+			
+			int32 DimIndex = Axes[Idx];
+			check(DimIndex < InputTensors[0].GetShape().Rank());
+			int32 Stride = Steps[Idx];
+			check(Stride != 0);
+
+			int32 Dim = InputTensors[0].GetShape().GetData()[DimIndex];
+			if(Start < 0 && Start > TNumericLimits<int32>::Min())
+			{
+				Start += Dim;
+			}
+			if(End < 0 && Start > TNumericLimits<int32>::Min())
+			{
+				End += Dim;
+			}
+
+			if (Stride < 0)
+            {
+                std::swap(Start, End);
+                Start += (Start < TNumericLimits<int32>::Max()) ? 1 : 0;
+                End += (End < TNumericLimits<int32>::Max()) ? 1 : 0;
+            }
+
+			Start = FMath::Max(Start, 0);
+            End = FMath::Min(End, Dim);
+            int32 Size = FMath::Max(End - Start, 0);
+			
+			int32 AbsStride = FMath::Abs(Stride);
+            OutputShape[DimIndex] = (uint32) ((Size / AbsStride) + (Size % AbsStride != 0));
+            Offsets[DimIndex] = (uint32) Start;
+            Strides[DimIndex] = Stride;
+            Sizes[DimIndex] = (uint32) Size;
+		}
+
+
+		DML_SLICE1_OPERATOR_DESC DmlSliceOpDesc{};
+		
+		DmlSliceOpDesc.InputTensor = &DmlInputTensorDesc.Desc;
+        DmlSliceOpDesc.OutputTensor = &DmlOutputTensorDesc.Desc;
+        DmlSliceOpDesc.DimensionCount = (uint32) Offsets.Num();
+        DmlSliceOpDesc.InputWindowOffsets = Offsets.GetData();
+        DmlSliceOpDesc.InputWindowSizes = Sizes.GetData();
+        DmlSliceOpDesc.InputWindowStrides = Strides.GetData();
+
+		return CreateOperator(Device, DML_OPERATOR_DESC{ DML_OPERATOR_SLICE1, &DmlSliceOpDesc} );
+
+	}
+};
+
+// Register Slice operator on Module startup
+NNE_DML_REGISTER_OP(Slice)
+
+} // namespace UE::NNERuntimeRDG::Private::Dml
+
+#endif // NNE_USE_DIRECTML
