NNE Binary operator shape inference : Add support for int32 and int64 input tensors

Tests: Added related unit tests for mul operator

#rnx

diff --git a/Engine/Plugins/Experimental/NNE/Source/NNEQA/Private/NNEQAUnitTestOpBinary.cpp b/Engine/Plugins/Experimental/NNE/Source/NNEQA/Private/NNEQAUnitTestOpBinary.cpp
--- a/Engine/Plugins/Experimental/NNE/Source/NNEQA/Private/NNEQAUnitTestOpBinary.cpp
+++ b/Engine/Plugins/Experimental/NNE/Source/NNEQA/Private/NNEQAUnitTestOpBinary.cpp
@@ -130,6 +130,80 @@ namespace UE::NNEQA::Private::NNERuntimeRDG::BinaryOp
 		return true;
 	};
 
+	IMPLEMENT_NNE_SHAPEINFERENCEHELPER_UNIT_AUTOMATION_TEST(FElementWiseBinaryCPUHelperMulInt32, "System.Engine.MachineLearning.NNE.UnitTest.BinaryHelper.MulInt32");
+	bool FElementWiseBinaryCPUHelperMulInt32::RunTest(const FString& Parameter)
+	{
+		FTensor XC1 = MakeConstTensorInt32(TEXT("XC1"), { 1 }, { 1 });
+		FTensor XC1x2 = MakeConstTensorInt32(TEXT("XC1x2"), { 1,2 }, { 1, 2 });
+		FTensor XC2x1 = MakeConstTensorInt32(TEXT("XC2x1"), { 2,1 }, { 3, 4 });
+
+		if (!TestBinaryOutputIsOnlyComputedWhenItShould(EElementWiseBinaryOperatorType::Mul))
+		{
+			return false;
+		}
+
+		//Test output tensor math is correct when both inputs are constant including broadcasting
+		FTensor Y = MakeTensor(TEXT("Y"), { 1 }, ENNETensorDataType::Int32);
+		CPUHelper::ElementWiseBinary::Apply(EElementWiseBinaryOperatorType::Mul, XC1, XC1, Y);
+		UTEST_EQUAL(TEXT("XC1*XC1"), Y.GetPreparedData<int32>()[0], 1 * 1);
+
+		Y = MakeTensor(TEXT("Y"), { 1,2 }, ENNETensorDataType::Int32);
+		CPUHelper::ElementWiseBinary::Apply(EElementWiseBinaryOperatorType::Mul, XC1x2, XC1, Y);
+		UTEST_EQUAL(TEXT("XC1x2*XC1[0]"), Y.GetPreparedData<int32>()[0], 1 * 1);
+		UTEST_EQUAL(TEXT("XC1x2*XC1[1]"), Y.GetPreparedData<int32>()[1], 2 * 1);
+
+		Y = MakeTensor(TEXT("Y"), { 1,2 }, ENNETensorDataType::Int32);
+		CPUHelper::ElementWiseBinary::Apply(EElementWiseBinaryOperatorType::Mul, XC1x2, XC1x2, Y);
+		UTEST_EQUAL(TEXT("XC1x2*XC1x2[0]"), Y.GetPreparedData<int32>()[0], 1 * 1);
+		UTEST_EQUAL(TEXT("XC1x2*XC1x2[1]"), Y.GetPreparedData<int32>()[1], 2 * 2);
+
+		Y = MakeTensor(TEXT("Y"), { 2,2 }, ENNETensorDataType::Int32);
+		CPUHelper::ElementWiseBinary::Apply(EElementWiseBinaryOperatorType::Mul, XC1x2, XC2x1, Y);
+		UTEST_EQUAL(TEXT("XC1x2*XC2x1[0]"), Y.GetPreparedData<int32>()[0], 1 * 3);
+		UTEST_EQUAL(TEXT("XC1x2*XC2x1[1]"), Y.GetPreparedData<int32>()[1], 2 * 3);
+		UTEST_EQUAL(TEXT("XC1x2*XC2x1[2]"), Y.GetPreparedData<int32>()[2], 1 * 4);
+		UTEST_EQUAL(TEXT("XC1x2*XC2x1[3]"), Y.GetPreparedData<int32>()[3], 2 * 4);
+
+		return true;
+	};
+
+	IMPLEMENT_NNE_SHAPEINFERENCEHELPER_UNIT_AUTOMATION_TEST(FElementWiseBinaryCPUHelperMulInt64, "System.Engine.MachineLearning.NNE.UnitTest.BinaryHelper.MulInt64");
+	bool FElementWiseBinaryCPUHelperMulInt64::RunTest(const FString& Parameter)
+	{
+		FTensor XC1 = MakeConstTensorInt64(TEXT("XC1"), { 1 }, { 1 });
+		FTensor XC1x2 = MakeConstTensorInt64(TEXT("XC1x2"), { 1,2 }, { 1, 2 });
+		FTensor XC2x1 = MakeConstTensorInt64(TEXT("XC2x1"), { 2,1 }, { 3, 4 });
+
+		if (!TestBinaryOutputIsOnlyComputedWhenItShould(EElementWiseBinaryOperatorType::Mul))
+		{
+			return false;
+		}
+
+		//Test output tensor math is correct when both inputs are constant including broadcasting
+		FTensor Y = MakeTensor(TEXT("Y"), { 1 }, ENNETensorDataType::Int64);
+		CPUHelper::ElementWiseBinary::Apply(EElementWiseBinaryOperatorType::Mul, XC1, XC1, Y);
+		UTEST_EQUAL(TEXT("XC1*XC1"), Y.GetPreparedData<int64>()[0], (int64)1 * 1);
+
+		Y = MakeTensor(TEXT("Y"), { 1,2 }, ENNETensorDataType::Int64);
+		CPUHelper::ElementWiseBinary::Apply(EElementWiseBinaryOperatorType::Mul, XC1x2, XC1, Y);
+		UTEST_EQUAL(TEXT("XC1x2*XC1[0]"), Y.GetPreparedData<int64>()[0], (int64)1 * 1);
+		UTEST_EQUAL(TEXT("XC1x2*XC1[1]"), Y.GetPreparedData<int64>()[1], (int64)2 * 1);
+
+		Y = MakeTensor(TEXT("Y"), { 1,2 }, ENNETensorDataType::Int64);
+		CPUHelper::ElementWiseBinary::Apply(EElementWiseBinaryOperatorType::Mul, XC1x2, XC1x2, Y);
+		UTEST_EQUAL(TEXT("XC1x2*XC1x2[0]"), Y.GetPreparedData<int64>()[0], (int64)1 * 1);
+		UTEST_EQUAL(TEXT("XC1x2*XC1x2[1]"), Y.GetPreparedData<int64>()[1], (int64)2 * 2);
+
+		Y = MakeTensor(TEXT("Y"), { 2,2 }, ENNETensorDataType::Int64);
+		CPUHelper::ElementWiseBinary::Apply(EElementWiseBinaryOperatorType::Mul, XC1x2, XC2x1, Y);
+		UTEST_EQUAL(TEXT("XC1x2*XC2x1[0]"), Y.GetPreparedData<int64>()[0], (int64)1 * 3);
+		UTEST_EQUAL(TEXT("XC1x2*XC2x1[1]"), Y.GetPreparedData<int64>()[1], (int64)2 * 3);
+		UTEST_EQUAL(TEXT("XC1x2*XC2x1[2]"), Y.GetPreparedData<int64>()[2], (int64)1 * 4);
+		UTEST_EQUAL(TEXT("XC1x2*XC2x1[3]"), Y.GetPreparedData<int64>()[3], (int64)2 * 4);
+
+		return true;
+	};
+
 	IMPLEMENT_NNE_SHAPEINFERENCEHELPER_UNIT_AUTOMATION_TEST(FElementWiseBinaryCPUHelperSub, "System.Engine.MachineLearning.NNE.UnitTest.BinaryHelper.Sub");
 	bool FElementWiseBinaryCPUHelperSub::RunTest(const FString& Parameter)
 	{
diff --git a/Engine/Plugins/Experimental/NNE/Source/NNERuntimeRDG/Private/Helper/NNERuntimeRDGHelperElementWiseBinary.cpp b/Engine/Plugins/Experimental/NNE/Source/NNERuntimeRDG/Private/Helper/NNERuntimeRDGHelperElementWiseBinary.cpp
--- a/Engine/Plugins/Experimental/NNE/Source/NNERuntimeRDG/Private/Helper/NNERuntimeRDGHelperElementWiseBinary.cpp
+++ b/Engine/Plugins/Experimental/NNE/Source/NNERuntimeRDG/Private/Helper/NNERuntimeRDGHelperElementWiseBinary.cpp
@@ -9,16 +9,33 @@
 
 namespace UE::NNERuntimeRDG::Internal::CPUHelper::ElementWiseBinary
 {
-	template<NNECore::Internal::EElementWiseBinaryOperatorType OpType> float Apply(float X, float Y);
-	template<> float Apply<NNECore::Internal::EElementWiseBinaryOperatorType::Add>(float X, float Y) { return X + Y; }
-	template<> float Apply<NNECore::Internal::EElementWiseBinaryOperatorType::Div>(float X, float Y) { return X / Y; }
-	template<> float Apply<NNECore::Internal::EElementWiseBinaryOperatorType::Mod>(float X, float Y) { return FMath::Fmod(X, Y); }
-	template<> float Apply<NNECore::Internal::EElementWiseBinaryOperatorType::Mul>(float X, float Y) { return X * Y; }
-	template<> float Apply<NNECore::Internal::EElementWiseBinaryOperatorType::Prelu>(float X, float Y) { return (X < 0.0f) ? (Y * X) : X; }
-	template<> float Apply<NNECore::Internal::EElementWiseBinaryOperatorType::Pow>(float X, float Y) { return FMath::Pow(X, Y); }
-	template<> float Apply<NNECore::Internal::EElementWiseBinaryOperatorType::Sub>(float X, float Y) { return X - Y; }
+	template<NNECore::Internal::EElementWiseBinaryOperatorType OpType, typename TData> TData Apply(TData X, TData Y);
 
-	template<NNECore::Internal::EElementWiseBinaryOperatorType OpType> void Apply(const NNECore::Internal::FTensor& LHSTensor, const NNECore::Internal::FTensor& RHSTensor, NNECore::Internal::FTensor& OutputTensor)
+	template<> float Apply<NNECore::Internal::EElementWiseBinaryOperatorType::Add, float>(float X, float Y) { return X + Y; }
+	template<> float Apply<NNECore::Internal::EElementWiseBinaryOperatorType::Div, float>(float X, float Y) { return X / Y; }
+	template<> float Apply<NNECore::Internal::EElementWiseBinaryOperatorType::Mod, float>(float X, float Y) { return FMath::Fmod(X, Y); }
+	template<> float Apply<NNECore::Internal::EElementWiseBinaryOperatorType::Mul, float>(float X, float Y) { return X * Y; }
+	template<> float Apply<NNECore::Internal::EElementWiseBinaryOperatorType::Prelu, float>(float X, float Y) { return (X < 0.0f) ? (Y * X) : X; }
+	template<> float Apply<NNECore::Internal::EElementWiseBinaryOperatorType::Pow, float>(float X, float Y) { return FMath::Pow(X, Y); }
+	template<> float Apply<NNECore::Internal::EElementWiseBinaryOperatorType::Sub, float>(float X, float Y) { return X - Y; }
+
+	template<> int32 Apply<NNECore::Internal::EElementWiseBinaryOperatorType::Add, int32>(int32 X, int32 Y) { return X + Y; }
+	template<> int32 Apply<NNECore::Internal::EElementWiseBinaryOperatorType::Div, int32>(int32 X, int32 Y) { return X / Y; }
+	template<> int32 Apply<NNECore::Internal::EElementWiseBinaryOperatorType::Mod, int32>(int32 X, int32 Y) { return X % Y; }
+	template<> int32 Apply<NNECore::Internal::EElementWiseBinaryOperatorType::Mul, int32>(int32 X, int32 Y) { return X * Y; }
+	template<> int32 Apply<NNECore::Internal::EElementWiseBinaryOperatorType::Prelu, int32>(int32 X, int32 Y) { return (X < 0.0f) ? (Y * X) : X; }
+	template<> int32 Apply<NNECore::Internal::EElementWiseBinaryOperatorType::Pow, int32>(int32 X, int32 Y) { return (int32)FMath::Pow((float)X, (float)Y); }
+	template<> int32 Apply<NNECore::Internal::EElementWiseBinaryOperatorType::Sub, int32>(int32 X, int32 Y) { return X - Y; }
+
+	template<> int64 Apply<NNECore::Internal::EElementWiseBinaryOperatorType::Add, int64>(int64 X, int64 Y) { return X + Y; }
+	template<> int64 Apply<NNECore::Internal::EElementWiseBinaryOperatorType::Div, int64>(int64 X, int64 Y) { return X / Y; }
+	template<> int64 Apply<NNECore::Internal::EElementWiseBinaryOperatorType::Mod, int64>(int64 X, int64 Y) { return X % Y; }
+	template<> int64 Apply<NNECore::Internal::EElementWiseBinaryOperatorType::Mul, int64>(int64 X, int64 Y) { return X * Y; }
+	template<> int64 Apply<NNECore::Internal::EElementWiseBinaryOperatorType::Prelu, int64>(int64 X, int64 Y) { return (X < 0.0f) ? (Y * X) : X; }
+	template<> int64 Apply<NNECore::Internal::EElementWiseBinaryOperatorType::Pow, int64>(int64 X, int64 Y) { return (int64)FMath::Pow((float)X, (float)Y); }
+	template<> int64 Apply<NNECore::Internal::EElementWiseBinaryOperatorType::Sub, int64>(int64 X, int64 Y) { return X - Y; }
+
+	template<NNECore::Internal::EElementWiseBinaryOperatorType OpType, typename TData> void Apply(const NNECore::Internal::FTensor& LHSTensor, const NNECore::Internal::FTensor& RHSTensor, NNECore::Internal::FTensor& OutputTensor)
 	{
 		//Heuristic to avoid unexpected performance hit. This helper being intended for shape related arithmetic only.
 		static constexpr int32 MaxItemInInputTensors = NNECore::FTensorShape::MaxRank * 2;
@@ -28,9 +45,9 @@ namespace UE::NNERuntimeRDG::Internal::CPUHelper::ElementWiseBinary
 			(LHSTensor.GetVolume() <= MaxItemInInputTensors) &&
 			(RHSTensor.GetVolume() <= MaxItemInInputTensors))
 		{
-			TConstArrayView<float> LHSData = LHSTensor.GetPreparedData<float>();
-			TConstArrayView<float> RHSData = RHSTensor.GetPreparedData<float>();
-			TArray<float> OutputData;
+			TConstArrayView<TData> LHSData = LHSTensor.GetPreparedData<TData>();
+			TConstArrayView<TData> RHSData = RHSTensor.GetPreparedData<TData>();
+			TArray<TData> OutputData;
 			OutputData.Reserve(OutputTensor.GetVolume());
 
 			Private::TensorIdxIterator it(OutputTensor.GetShape());
@@ -38,40 +55,58 @@ namespace UE::NNERuntimeRDG::Internal::CPUHelper::ElementWiseBinary
 			{
 				int32 LHSIdx = it.GetIndexToBroadcastedShape(LHSTensor.GetShape());
 				int32 RHSIdx = it.GetIndexToBroadcastedShape(RHSTensor.GetShape());
-				float LHSValue = LHSData.GetData()[LHSIdx];
-				float RHSValue = RHSData.GetData()[RHSIdx];
-				OutputData.Add(Apply<OpType>(LHSValue, RHSValue));
+				TData LHSValue = LHSData.GetData()[LHSIdx];
+				TData RHSValue = RHSData.GetData()[RHSIdx];
+				OutputData.Add(Apply<OpType, TData>(LHSValue, RHSValue));
 			} while (it.Advance());
 
 			check(OutputData.Num() == OutputTensor.GetVolume());
-			OutputTensor.SetPreparedData<float>(OutputData);
+			OutputTensor.SetPreparedData<TData>(OutputData);
 		}
 	}
 
-	void Apply(NNECore::Internal::EElementWiseBinaryOperatorType OpType, const NNECore::Internal::FTensor& LHSTensor, const NNECore::Internal::FTensor& RHSTensor, NNECore::Internal::FTensor& OutputTensor)
+	template<typename TData> void ApplyResolvedDataType(NNECore::Internal::EElementWiseBinaryOperatorType OpType, const NNECore::Internal::FTensor& LHSTensor, const NNECore::Internal::FTensor& RHSTensor, NNECore::Internal::FTensor& OutputTensor)
 	{
 		switch (OpType)
 		{
 		case NNECore::Internal::EElementWiseBinaryOperatorType::Add:
-			Apply<NNECore::Internal::EElementWiseBinaryOperatorType::Add>(LHSTensor, RHSTensor, OutputTensor);
+			Apply<NNECore::Internal::EElementWiseBinaryOperatorType::Add, TData>(LHSTensor, RHSTensor, OutputTensor);
 			break;
 		case NNECore::Internal::EElementWiseBinaryOperatorType::Div:
-			Apply<NNECore::Internal::EElementWiseBinaryOperatorType::Div>(LHSTensor, RHSTensor, OutputTensor);
+			Apply<NNECore::Internal::EElementWiseBinaryOperatorType::Div, TData>(LHSTensor, RHSTensor, OutputTensor);
 			break;
 		case NNECore::Internal::EElementWiseBinaryOperatorType::Mod:
-			Apply<NNECore::Internal::EElementWiseBinaryOperatorType::Mod>(LHSTensor, RHSTensor, OutputTensor);
+			Apply<NNECore::Internal::EElementWiseBinaryOperatorType::Mod, TData>(LHSTensor, RHSTensor, OutputTensor);
 			break;
 		case NNECore::Internal::EElementWiseBinaryOperatorType::Mul:
-			Apply<NNECore::Internal::EElementWiseBinaryOperatorType::Mul>(LHSTensor, RHSTensor, OutputTensor);
+			Apply<NNECore::Internal::EElementWiseBinaryOperatorType::Mul, TData>(LHSTensor, RHSTensor, OutputTensor);
 			break;
 		case NNECore::Internal::EElementWiseBinaryOperatorType::Prelu:
-			Apply<NNECore::Internal::EElementWiseBinaryOperatorType::Prelu>(LHSTensor, RHSTensor, OutputTensor);
+			Apply<NNECore::Internal::EElementWiseBinaryOperatorType::Prelu, TData>(LHSTensor, RHSTensor, OutputTensor);
 			break;
 		case NNECore::Internal::EElementWiseBinaryOperatorType::Pow:
-			Apply<NNECore::Internal::EElementWiseBinaryOperatorType::Pow>(LHSTensor, RHSTensor, OutputTensor);
+			Apply<NNECore::Internal::EElementWiseBinaryOperatorType::Pow, TData>(LHSTensor, RHSTensor, OutputTensor);
 			break;
 		case NNECore::Internal::EElementWiseBinaryOperatorType::Sub:
-			Apply<NNECore::Internal::EElementWiseBinaryOperatorType::Sub>(LHSTensor, RHSTensor, OutputTensor);
+			Apply<NNECore::Internal::EElementWiseBinaryOperatorType::Sub, TData>(LHSTensor, RHSTensor, OutputTensor);
+			break;
+		default:
+			break;
+		}
+	}
+
+	void Apply(NNECore::Internal::EElementWiseBinaryOperatorType OpType, const NNECore::Internal::FTensor& LHSTensor, const NNECore::Internal::FTensor& RHSTensor, NNECore::Internal::FTensor& OutputTensor)
+	{
+		switch (OutputTensor.GetDataType())
+		{
+		case ENNETensorDataType::Float:
+			ApplyResolvedDataType<float>(OpType, LHSTensor, RHSTensor, OutputTensor);
+			break;
+		case ENNETensorDataType::Int32:
+			ApplyResolvedDataType<int32>(OpType, LHSTensor, RHSTensor, OutputTensor);
+			break;
+		case ENNETensorDataType::Int64:
+			ApplyResolvedDataType<int64>(OpType, LHSTensor, RHSTensor, OutputTensor);
 			break;
 		default:
 			break;
