NNE Slice operator shape inference : Add support for int32 and int64 input tensors

Tests: Added related unit tests

#rnx

diff --git a/Engine/Plugins/Experimental/NNE/Source/NNEQA/Private/NNEQAUnitTestOpSlice.cpp b/Engine/Plugins/Experimental/NNE/Source/NNEQA/Private/NNEQAUnitTestOpSlice.cpp
--- a/Engine/Plugins/Experimental/NNE/Source/NNEQA/Private/NNEQAUnitTestOpSlice.cpp
+++ b/Engine/Plugins/Experimental/NNE/Source/NNEQA/Private/NNEQAUnitTestOpSlice.cpp
@@ -70,6 +70,27 @@ namespace UE::NNEQA::Private::NNERuntimeRDG::SliceOp
 		return true;
 	}
 
+	IMPLEMENT_NNE_SHAPEINFERENCEHELPER_UNIT_AUTOMATION_TEST(FSliceCPUHelperRank1Int, "System.Engine.MachineLearning.NNE.UnitTest.SliceHelper.Rank1Int");
+	bool FSliceCPUHelperRank1Int::RunTest(const FString& Parameter)
+	{
+		FTensor XC6Int32 = MakeConstTensorInt32(TEXT("XC6Int32"), { 6 }, { 1, 2, 3, 4, 5, 6 });
+		FTensor XC6Int64 = MakeConstTensorInt64(TEXT("XC6Int64"), { 6 }, { 1, 2, 3, 4, 5, 6 });
+		FTensor YInt32 = MakeTensor(TEXT("YInt32"), { 2 }, ENNETensorDataType::Int32);
+		FTensor YInt64 = MakeTensor(TEXT("YInt64"), { 2 }, ENNETensorDataType::Int64);
+
+		CPUHelper::Slice::Apply(XC6Int32, YInt32, { 1 });
+		UTEST_TRUE(TEXT("YInt32 const if input is const"), YInt32.HasPreparedData());
+		UTEST_EQUAL(TEXT("Slice(XC6Int32,2,4)[0]"), YInt32.GetPreparedData<int32>()[0], 2);
+		UTEST_EQUAL(TEXT("Slice(XC6Int32,2,4)[1]"), YInt32.GetPreparedData<int32>()[1], 3);
+
+		CPUHelper::Slice::Apply(XC6Int64, YInt64, { 1 });
+		UTEST_TRUE(TEXT("YInt64 const if input is const"), YInt64.HasPreparedData());
+		UTEST_EQUAL(TEXT("Slice(XC6Int64,2,4)[0]"), YInt64.GetPreparedData<int64>()[0], (int64)2);
+		UTEST_EQUAL(TEXT("Slice(XC6Int64,2,4)[1]"), YInt64.GetPreparedData<int64>()[1], (int64)3);
+
+		return true;
+	}
+
 	IMPLEMENT_NNE_SHAPEINFERENCEHELPER_UNIT_AUTOMATION_TEST(FSliceCPUHelperRank3, "System.Engine.MachineLearning.NNE.UnitTest.SliceHelper.Rank3");
 	bool FSliceCPUHelperRank3::RunTest(const FString& Parameter)
 	{
diff --git a/Engine/Plugins/Experimental/NNE/Source/NNERuntimeRDG/Private/Helper/NNERuntimeRDGHelperSlice.cpp b/Engine/Plugins/Experimental/NNE/Source/NNERuntimeRDG/Private/Helper/NNERuntimeRDGHelperSlice.cpp
--- a/Engine/Plugins/Experimental/NNE/Source/NNERuntimeRDG/Private/Helper/NNERuntimeRDGHelperSlice.cpp
+++ b/Engine/Plugins/Experimental/NNE/Source/NNERuntimeRDG/Private/Helper/NNERuntimeRDGHelperSlice.cpp
@@ -5,30 +5,14 @@
 
 namespace UE::NNERuntimeRDG::Internal::CPUHelper::Slice
 {
-	void Apply(const NNECore::Internal::FTensor& InputTensor, NNECore::Internal::FTensor& OutputTensor, TConstArrayView<int32> Starts)
+	template<typename TData> void ApplyResolvedInputType(const NNECore::Internal::FTensor& InputTensor, NNECore::Internal::FTensor& OutputTensor, TConstArrayView<int32> Starts)
 	{
-		static constexpr int32 MaxItemInOutputTensor = NNECore::FTensorShape::MaxRank * 2;
-
-		if (OutputTensor.GetVolume() >= MaxItemInOutputTensor)
-		{
-			return;
-		}
-
-		if (OutputTensor.GetDataType() != ENNETensorDataType::Float)
-		{
-			return;
-		}
-
-		if (!InputTensor.HasPreparedData())
-		{
-			return;
-		}
-
+		check(InputTensor.HasPreparedData());
 		check(InputTensor.GetShape().Rank() == Starts.Num());
 		check(OutputTensor.GetShape().Rank() == Starts.Num());
 
-		TArray<float> OutputData;
-		TConstArrayView<float> InputData = InputTensor.GetPreparedData<float>();
+		TArray<TData> OutputData;
+		TConstArrayView<TData> InputData = InputTensor.GetPreparedData<TData>();
 		Private::TensorIdxIterator itOutput(OutputTensor.GetShape());
 		const Private::TensorIdxIterator itInput(InputTensor.GetShape());
 
@@ -41,12 +25,42 @@ namespace UE::NNERuntimeRDG::Internal::CPUHelper::Slice
 				CurInputPosition[r] += Starts[r];
 			}
 
-			float Value = InputData[itInput.GetIndexFromPosition(CurInputPosition)];
+			TData Value = InputData[itInput.GetIndexFromPosition(CurInputPosition)];
 			OutputData[itOutput.GetIndex()] = Value;
 
 		} while (itOutput.Advance());
 
-		OutputTensor.SetPreparedData<float>(OutputData);
+		OutputTensor.SetPreparedData<TData>(OutputData);
+	}
+
+	void Apply(const NNECore::Internal::FTensor& InputTensor, NNECore::Internal::FTensor& OutputTensor, TConstArrayView<int32> Starts)
+	{
+		static constexpr int32 MaxItemInOutputTensor = NNECore::FTensorShape::MaxRank * 2;
+
+		if (OutputTensor.GetVolume() >= MaxItemInOutputTensor)
+		{
+			return;
+		}
+
+		if (!InputTensor.HasPreparedData())
+		{
+			return;
+		}
+
+		switch (InputTensor.GetDataType())
+		{
+		case ENNETensorDataType::Int32:
+			ApplyResolvedInputType<int32>(InputTensor, OutputTensor, Starts);
+			break;
+
+		case ENNETensorDataType::Int64:
+			ApplyResolvedInputType<int64>(InputTensor, OutputTensor, Starts);
+			break;
+
+		case ENNETensorDataType::Float:
+			ApplyResolvedInputType<float>(InputTensor, OutputTensor, Starts);
+			break;
+		}
 	}
 	
 } // UE::NNERuntimeRDG::Private::CPUHelper::Slice
diff --git a/Engine/Plugins/Experimental/NNE/Source/NNERuntimeRDG/Private/Hlsl/NNERuntimeRDGSlice.cpp b/Engine/Plugins/Experimental/NNE/Source/NNERuntimeRDG/Private/Hlsl/NNERuntimeRDGSlice.cpp
--- a/Engine/Plugins/Experimental/NNE/Source/NNERuntimeRDG/Private/Hlsl/NNERuntimeRDGSlice.cpp
+++ b/Engine/Plugins/Experimental/NNE/Source/NNERuntimeRDG/Private/Hlsl/NNERuntimeRDGSlice.cpp
@@ -140,6 +140,8 @@ namespace UE::NNERuntimeRDG::Private::Hlsl
 		bIsValid &= AttributeValidator.Validate(AttributeMap);
 
 		FInputValidator InputValidator;
+		InputValidator.AddSupportedType(ENNETensorDataType::Int64);
+		InputValidator.AddSupportedType(ENNETensorDataType::Int32);
 		InputValidator.AddSupportedType(ENNETensorDataType::Float);
 		InputValidator.AddRequired();
 
