NNE - Add DML MatMul

#rnx

diff --git a/Engine/Plugins/Experimental/NNE/Source/NNEQA/Resources/NNEQATestDesc.json b/Engine/Plugins/Experimental/NNE/Source/NNEQA/Resources/NNEQATestDesc.json
--- a/Engine/Plugins/Experimental/NNE/Source/NNEQA/Resources/NNEQATestDesc.json
+++ b/Engine/Plugins/Experimental/NNE/Source/NNEQA/Resources/NNEQATestDesc.json
@@ -316,7 +316,7 @@
 		},
 		{
 			"category": "MatMul",
-			"runtimes": [ { "name": "NNERuntimeRDGDml", "skip": true } ],
+			"runtimes": [ { "name": "NNERuntimeRDGDml" } ],
 			"targets": [
 				{ "target": "MatMul" }
 			]
@@ -633,8 +633,8 @@
 				{ "name": "NNERuntimeRDGHlsl"}
 			],
 			"datasets": [
-				{ "inputs": [ { "shape": [ 2, 2, 3, 5, 4 ] }, { "shape": [ 2, 1, 3, 4, 3 ] } ], "outputs": [ { "shape": [ 2, 2, 3, 5, 3 ] } ]},
-				{ "inputs": [ { "shape": [ 2, 1, 2, 5, 4 ] }, { "shape": [ 1, 3, 2, 4, 3 ] } ], "outputs": [ { "shape": [ 2, 3, 2, 5, 3 ] } ]},
+				{ "inputs": [ { "shape": [ 2, 2, 3, 5, 4 ] }, { "shape": [ 2, 1, 3, 4, 3 ] } ], "outputs": [ { "shape": [ 2, 2, 3, 5, 3 ] } ], "runtimes": [ {"name": "NNERuntimeRDGDml", "skip": true, "comment": "Not Supported"} ]},
+				{ "inputs": [ { "shape": [ 2, 1, 2, 5, 4 ] }, { "shape": [ 1, 3, 2, 4, 3 ] } ], "outputs": [ { "shape": [ 2, 3, 2, 5, 3 ] } ], "runtimes": [ {"name": "NNERuntimeRDGDml", "skip": true, "comment": "Not Supported"} ]},
 				{ "inputs": [ { "shape": [ 1, 1, 2, 5, 4 ] }, { "shape": [ 2, 3, 2, 4, 3 ] } ], "outputs": [ { "shape": [ 2, 3, 2, 5, 3 ] } ]},
 				{ "inputs": [ { "shape": [ 1, 2, 5, 4 ] }, { "shape": [ 2, 3, 2, 4, 3 ] } ], "outputs": [ { "shape": [ 2, 3, 2, 5, 3 ] } ]},
 				{ "inputs": [ { "shape": [ 5, 4 ] }, { "shape": [ 2, 1, 3, 4, 3 ] } ], "outputs": [ { "shape": [ 2, 1, 3, 5, 3 ] } ]},
diff --git a/Engine/Plugins/Experimental/NNE/Source/NNERuntimeRDG/Private/Dml/NNEDmlOperator.cpp b/Engine/Plugins/Experimental/NNE/Source/NNERuntimeRDG/Private/Dml/NNEDmlOperator.cpp
--- a/Engine/Plugins/Experimental/NNE/Source/NNERuntimeRDG/Private/Dml/NNEDmlOperator.cpp
+++ b/Engine/Plugins/Experimental/NNE/Source/NNERuntimeRDG/Private/Dml/NNEDmlOperator.cpp
@@ -33,7 +33,7 @@ namespace DmlUtil
 		}
 	}
 
-	bool FTensorDesc::InitFromTensor(const NNECore::Internal::FTensor& Tensor, int32 MinTensorRank, TConstArrayView<uint32> BroadcastShape)
+	bool FTensorDesc::InitFromTensor(const NNECore::Internal::FTensor& Tensor, int32 MinTensorRank, TConstArrayView<uint32> BroadcastShape, TConstArrayView<uint32> CustomShape)
 	{
 		Reset();
 
@@ -44,7 +44,7 @@ namespace DmlUtil
 			return false;
 		}
 
-		TConstArrayView<uint32> InShape = Tensor.GetShape().GetData();
+		TConstArrayView<uint32> InShape = CustomShape.IsEmpty() ? Tensor.GetShape().GetData() : CustomShape;
 
 		ElemSizeInBytes = Tensor.GetElemByteSize();
 
diff --git a/Engine/Plugins/Experimental/NNE/Source/NNERuntimeRDG/Private/Dml/NNEDmlOperator.h b/Engine/Plugins/Experimental/NNE/Source/NNERuntimeRDG/Private/Dml/NNEDmlOperator.h
--- a/Engine/Plugins/Experimental/NNE/Source/NNERuntimeRDG/Private/Dml/NNEDmlOperator.h
+++ b/Engine/Plugins/Experimental/NNE/Source/NNERuntimeRDG/Private/Dml/NNEDmlOperator.h
@@ -57,7 +57,7 @@ namespace DmlUtil
 		FSmallUIntArray			Strides;
 		uint64					ElemSizeInBytes;
 
-		bool InitFromTensor(const NNECore::Internal::FTensor& InputDesc, int32 MinTensorRank, TConstArrayView<uint32> Broadcast = MakeArrayView((uint32*) nullptr, 0));
+		bool InitFromTensor(const NNECore::Internal::FTensor& InputDesc, int32 MinTensorRank, TConstArrayView<uint32> Broadcast = MakeArrayView((uint32*) nullptr, 0), TConstArrayView<uint32> CustomShape = MakeArrayView((uint32*) nullptr, 0));
 		bool InitFromTensor1D(const NNECore::Internal::FTensor& InputDesc, int32 Rank);
 
 		void SetStridesFromFTensor(const NNECore::Internal::FTensor& InputDesc);
diff --git a/Engine/Plugins/Experimental/NNE/Source/NNERuntimeRDG/Private/Dml/NNEDmlOperatorMatMul.cpp b/Engine/Plugins/Experimental/NNE/Source/NNERuntimeRDG/Private/Dml/NNEDmlOperatorMatMul.cpp
new file mode 100644
--- /dev/null
+++ b/Engine/Plugins/Experimental/NNE/Source/NNERuntimeRDG/Private/Dml/NNEDmlOperatorMatMul.cpp
@@ -0,0 +1,237 @@
+// Copyright Epic Games, Inc. All Rights Reserved.
+
+#ifdef NNE_USE_DIRECTML
+#include "NNEDmlOperator.h"
+
+namespace UE::NNERuntimeRDG::Private::Dml
+{
+
+/**
+ * MatMul
+ */
+class FOperatorDmlMatMul : public FOperatorDml
+{
+public:
+
+	static FOperatorDml* Create()
+	{
+		return new FOperatorDmlMatMul();
+	}
+
+	//
+	//
+	//
+	virtual bool Initialize(IDMLDevice* Device, TArrayView<const NNECore::Internal::FTensor> InputTensors, TArrayView<const NNECore::Internal::FTensor> OutputTensors, const NNECore::FAttributeMap& Attributes) override
+	{
+		check(InputTensors.Num() == 2);
+		check(OutputTensors.Num() == 1);
+
+		using namespace UE::NNECore;
+
+		const NNECore::Internal::FTensor& ATensor = InputTensors[0];
+		const NNECore::Internal::FTensor& BTensor = InputTensors[1];
+		const NNECore::Internal::FTensor& OutTensor = OutputTensors[0];
+
+
+		DmlUtil::FSmallUIntArray ActualShapeA (ATensor.GetShape().GetData());
+		DmlUtil::FSmallUIntArray ActualShapeB (BTensor.GetShape().GetData());
+		DmlUtil::FSmallUIntArray ActualShapeOut (OutTensor.GetShape().GetData());
+
+		check(ActualShapeA.Num() > 0);
+		check(ActualShapeB.Num() > 0);
+		check(ActualShapeOut.Num() > 0);
+
+		if(ActualShapeA.Num() == 1)
+		{
+			check(ActualShapeOut.Num() == 1);
+			ActualShapeA = {1, ActualShapeA[0]};
+		}
+		if(ActualShapeB.Num() == 1)
+		{
+			check(ActualShapeOut.Num() == 1);
+			ActualShapeB = {ActualShapeB[0], 1};
+		}
+
+		// Since DML requires a 4D tensor as shape, where first two sizes are BatchCount and ChannelCount, we cannot represent general
+		// broadcasting of both A and B.
+		// Example: [2,1,2,5,4] x [1,3,2,4,3] => [2,3,2,5,3]
+		// Furthermore, for simplicity, we have chosen to only use the channel dimension to represent the stack index and put the batch size to 1.
+		// This implies that we support only the cases in which only one between A and B are broadcasted.
+		// This constraint could be relaxed by also smartly employing the batch dimension, but general broadcasting of both A and B is still not possible.
+
+		// Test compatibility (see also comments below)
+		bool bARequiresBroadcasting = false;
+		bool bBRequiresBroadcasting = false;
+		{
+			int32 MaxRank = FMath::Max(ActualShapeA.Num(), ActualShapeB.Num());
+			for(int Dim = MaxRank - 3; Dim >= 0; Dim--)
+			{
+				int AIdx = Dim - MaxRank + ActualShapeA.Num();
+				int BIdx = Dim - MaxRank + ActualShapeB.Num();
+
+				if(bARequiresBroadcasting)
+				{
+					if(AIdx >= 0)
+					{
+						check(ActualShapeA[AIdx] == 1);
+					}
+				}
+
+				if(bBRequiresBroadcasting)
+				{
+					if(BIdx >= 0)
+					{
+						check(ActualShapeB[BIdx] == 1);
+					}
+				}
+
+				if(AIdx >= 0 && BIdx >= 0)
+				{
+					if(ActualShapeA[AIdx] > ActualShapeB[BIdx])
+					{
+						check(!bARequiresBroadcasting);
+						check(ActualShapeA[AIdx] % ActualShapeB[BIdx] == 0);
+						bBRequiresBroadcasting = true;
+					}
+
+					if(ActualShapeB[BIdx] > ActualShapeA[AIdx])
+					{
+						check(!bBRequiresBroadcasting);
+						check(ActualShapeB[BIdx] % ActualShapeA[AIdx] == 0);
+						bARequiresBroadcasting = true;
+					}
+				}
+
+				if(AIdx < 0 && ActualShapeB[BIdx] != 1)
+				{
+					bARequiresBroadcasting = true;
+				}
+
+				if(BIdx < 0 && ActualShapeA[AIdx] != 1)
+				{
+					bBRequiresBroadcasting = true;
+				}
+				
+			}
+		}
+
+		
+		DmlUtil::FSmallUIntArray BroadcastShape(
+			bARequiresBroadcasting ? 
+			TConstArrayView<uint32>(ActualShapeB).LeftChop(2) : 
+			TConstArrayView<uint32>(ActualShapeA).LeftChop(2)
+			);
+
+		
+		check(BroadcastShape == TConstArrayView<uint32>(ActualShapeOut).LeftChop(2));
+
+		auto To4DTensor = [](TConstArrayView<uint32> BroadcastShape, TConstArrayView<uint32> MatrixShape) -> DmlUtil::FSmallUIntArray
+		{
+			check(MatrixShape.Num() == 2);
+
+			uint32 NumMatrices = 1;
+			for(uint32 Elem : BroadcastShape)
+			{
+				NumMatrices *= Elem;
+			}
+			return DmlUtil::FSmallUIntArray({1, NumMatrices, MatrixShape[0], MatrixShape[1]});
+		};
+
+		DmlUtil::FSmallUIntArray DmlShapeA = To4DTensor(TConstArrayView<uint32>(ActualShapeA).LeftChop(2), TConstArrayView<uint32>(ActualShapeA).Right(2));
+		DmlUtil::FSmallUIntArray BroadcastShapeA = To4DTensor(BroadcastShape, TConstArrayView<uint32>(ActualShapeA).Right(2));
+		DmlUtil::FSmallUIntArray DmlShapeB = To4DTensor(TConstArrayView<uint32>(ActualShapeB).LeftChop(2), TConstArrayView<uint32>(ActualShapeB).Right(2));
+		DmlUtil::FSmallUIntArray BroadcastShapeB = To4DTensor(BroadcastShape, TConstArrayView<uint32>(ActualShapeB).Right(2));
+		DmlUtil::FSmallUIntArray DmlShapeOut = To4DTensor(BroadcastShape, TConstArrayView<uint32>(ActualShapeOut).Right(2));
+
+		check(DmlShapeA[3] == DmlShapeB[2]); // K
+		check(DmlShapeOut[2] == DmlShapeA[2]); // M
+		check(DmlShapeOut[3] == DmlShapeB[3]); // N
+
+		// Another broadcast issue!
+		// The current implementation of DmlUtil::FTensorDesc::InitFromTensor() is not correct with respect to broadcast as it sets
+		// wrong sizes and strides when a dimension to be broadcasted is not of size 1.
+		// Example:
+		// [1, 2, 5, 4] ->(broadcast) [1, 12, 5, 4]
+		// would require
+		// sizes:   [ 1,   6,   2,   5,   4 ],
+		// strides: [ 40,  0,  20,   4,   1 ].
+		// This is in general possible only if 
+		// - one is free to have arbitrary sized arrays for sizes and strides, but DirectML constrains them to 4D or 5D,
+		// - the broadcast dimension size is a multiple of the broadcasted one.
+		// In our specific case here we can use both the batch and index dimensions to represent non-1 broadcasting, but
+		// we are still limited to cases where the broadcast dimension size is a multiple of the broadcasted one.
+		{
+			auto FixBroadcast = [](DmlUtil::FSmallUIntArray& DmlShape, DmlUtil::FSmallUIntArray& BroadcastArray, uint32 BroadcastFactor)
+			{
+				if(DmlShape == BroadcastArray)
+				{
+					DmlShape[0] = BroadcastFactor;
+					DmlShape[1] /= BroadcastFactor;
+				}
+
+				BroadcastArray[0] = BroadcastFactor;
+				BroadcastArray[1] /= BroadcastFactor;
+			};
+			check(DmlShapeB == BroadcastShapeB ? (BroadcastShapeA[1] % DmlShapeA[1] == 0) : (BroadcastShapeB[1] % DmlShapeB[1] == 0));
+			uint32 BroadcastFactor = DmlShapeB == BroadcastShapeB ? BroadcastShapeA[1] / DmlShapeA[1] : BroadcastShapeB[1] / DmlShapeB[1];
+
+			FixBroadcast(DmlShapeA, BroadcastShapeA, BroadcastFactor);
+			FixBroadcast(DmlShapeB, BroadcastShapeB, BroadcastFactor);
+			DmlShapeOut[0] = BroadcastFactor;
+			DmlShapeOut[1] /= BroadcastFactor;
+		}
+		
+		// Note: DmlShape* are the the 4D ones that InitFromTensor() would need to use, however can't set them in original tensor like below because of prepared data check.
+		//ATensor.SetShape(DmlShapeA);
+		//BTensor.SetShape(DmlShapeB);
+
+		DmlUtil::FTensorDesc	DmlInputATensorDesc{};
+		DmlUtil::FTensorDesc	DmlInputBTensorDesc{};
+		DmlUtil::FTensorDesc	DmlOutputTensorDesc{};
+
+		if (!DmlInputATensorDesc.InitFromTensor(ATensor, DmlShapeA.Num(), 
+			/*Broadcast =*/ BroadcastShapeA, 
+			/*CustomShape =*/ DmlShapeA))
+		{
+			UE_LOG(LogNNE, Warning, TEXT("Failed to initialize MatMul's tensor A for DML inference"));
+			return false;
+		}
+
+		if (!DmlInputBTensorDesc.InitFromTensor(BTensor, DmlShapeB.Num(), 
+			/*Broadcast =*/ BroadcastShapeB,
+			/*CustomShape =*/ DmlShapeB))
+		{
+			UE_LOG(LogNNE, Warning, TEXT("Failed to initialize MatMul's tensor B for DML inference"));
+			return false;
+		}
+
+		if (!DmlOutputTensorDesc.InitFromTensor(OutTensor, DmlShapeOut.Num(), 
+			/*Broadcast =*/ MakeArrayView((uint32*) nullptr, 0), 
+			/*CustomShape =*/ DmlShapeOut))
+		{
+			UE_LOG(LogNNE, Warning, TEXT("Failed to initialize MatMul's output tensor for DML inference"));
+			return false;
+		}
+
+		DML_GEMM_OPERATOR_DESC	DmlGemmOpDesc{};
+
+		DmlGemmOpDesc.ATensor = &DmlInputATensorDesc.Desc;
+		DmlGemmOpDesc.BTensor = &DmlInputBTensorDesc.Desc;
+		DmlGemmOpDesc.CTensor = nullptr;
+		DmlGemmOpDesc.OutputTensor = &DmlOutputTensorDesc.Desc;
+		DmlGemmOpDesc.Alpha = 1.0f;
+		DmlGemmOpDesc.Beta = 0.0f;
+		DmlGemmOpDesc.TransA = DML_MATRIX_TRANSFORM_NONE;
+		DmlGemmOpDesc.TransB = DML_MATRIX_TRANSFORM_NONE;
+		DmlGemmOpDesc.FusedActivation = nullptr;
+
+		return CreateOperator(Device, DML_OPERATOR_DESC{ DML_OPERATOR_GEMM, &DmlGemmOpDesc} );
+	}
+
+};
+
+NNE_DML_REGISTER_OP(MatMul)
+
+} // namespace UE::NNERuntimeRDG::Private::Dml
+
+#endif // NNE_USE_DIRECTML
