- Add RDG DirectML Upsample operator
- Refactor Dml FTensorDesc (update Tranpose to use refactored methods)
- Fix white space in InstanceNormalization

#rnx

diff --git a/Engine/Plugins/Experimental/NNE/Source/NNERuntimeRDG/Private/Dml/NNEDmlOperator.cpp b/Engine/Plugins/Experimental/NNE/Source/NNERuntimeRDG/Private/Dml/NNEDmlOperator.cpp
--- a/Engine/Plugins/Experimental/NNE/Source/NNERuntimeRDG/Private/Dml/NNEDmlOperator.cpp
+++ b/Engine/Plugins/Experimental/NNE/Source/NNERuntimeRDG/Private/Dml/NNEDmlOperator.cpp
@@ -11,28 +11,6 @@ namespace UE::NNERuntimeRDG::Private::Dml
 namespace DmlUtil
 {
 
-	void FTensorDesc::UpdateShapesAndStrides(TConstArrayView<uint32> InSizes, TConstArrayView<uint32> InStrides)
-	{
-		Sizes = InSizes;
-        BuffDesc.Sizes = Sizes.GetData();
-
-		Strides = InStrides;
-        BuffDesc.Strides = Strides.GetData();
-	}
-
-	void FTensorDesc::SetStridesFromFTensor(const NNECore::Internal::FTensor& InputDesc)
-	{
-		uint32 CurrStride = 1;
-
-		Strides.SetNum(InputDesc.GetShape().Rank());
-		
-		for (int32 i = InputDesc.GetShape().Rank() - 1; i >= 0; --i)
-		{
-			Strides[i] = CurrStride;
-			CurrStride *= InputDesc.GetShape().GetData()[i];
-		}
-	}
-
 	bool FTensorDesc::InitFromTensor(const NNECore::Internal::FTensor& Tensor, int32 MinTensorRank, TConstArrayView<uint32> BroadcastShape, TConstArrayView<uint32> CustomShape)
 	{
 		Reset();
@@ -89,6 +67,39 @@ namespace DmlUtil
 		return true;
 	}
 
+	void FTensorDesc::UpdateShapeAndStrides(TConstArrayView<uint32> InShape, TConstArrayView<uint32> InStrides)
+	{
+		check(!InStrides.IsEmpty() && (InStrides.Num() == InShape.Num()));
+
+		Sizes = InShape;
+		BuffDesc.Sizes = Sizes.GetData();
+		BuffDesc.DimensionCount = Sizes.Num();
+
+		if (!InStrides.IsEmpty())
+		{
+			Strides = InStrides;
+			BuffDesc.Strides = nullptr;
+		}
+		else
+		{
+			Strides.Reset();
+			BuffDesc.Strides = nullptr;
+		}
+	}
+
+	void FTensorDesc::SetStridesFromTensor(const NNECore::Internal::FTensor& InputDesc)
+	{
+		uint32 CurrStride = 1;
+
+		Strides.SetNum(InputDesc.GetShape().Rank());
+		
+		for (int32 i = InputDesc.GetShape().Rank() - 1; i >= 0; --i)
+		{
+			Strides[i] = CurrStride;
+			CurrStride *= InputDesc.GetShape().GetData()[i];
+		}
+	}
+
 	void FTensorDesc::Reset()
 	{
 		BuffDesc = DML_BUFFER_TENSOR_DESC{};
@@ -336,7 +347,6 @@ namespace DmlUtil
 //
 // DirectML operator base class
 //
-
 TConstArrayView<int32> FOperatorDml::GetConstantCPUInputs() const
 {
 	return ConstantCPUInputs;
diff --git a/Engine/Plugins/Experimental/NNE/Source/NNERuntimeRDG/Private/Dml/NNEDmlOperator.h b/Engine/Plugins/Experimental/NNE/Source/NNERuntimeRDG/Private/Dml/NNEDmlOperator.h
--- a/Engine/Plugins/Experimental/NNE/Source/NNERuntimeRDG/Private/Dml/NNEDmlOperator.h
+++ b/Engine/Plugins/Experimental/NNE/Source/NNERuntimeRDG/Private/Dml/NNEDmlOperator.h
@@ -45,6 +45,8 @@ public:
 
 namespace DmlUtil
 {
+	template<typename T>
+	using FSmallArray = TArray<T, TInlineAllocator<NNECore::FTensorShape::MaxRank>>;
 	using FSmallIntArray = TArray<int32, TInlineAllocator<NNECore::FTensorShape::MaxRank>>;
 	using FSmallUIntArray = TArray<uint32, TInlineAllocator<NNECore::FTensorShape::MaxRank>>;
 
@@ -60,9 +62,9 @@ namespace DmlUtil
 		bool InitFromTensor(const NNECore::Internal::FTensor& InputDesc, int32 MinTensorRank, TConstArrayView<uint32> Broadcast = MakeArrayView((uint32*) nullptr, 0), TConstArrayView<uint32> CustomShape = MakeArrayView((uint32*) nullptr, 0));
 		bool InitFromTensor1D(const NNECore::Internal::FTensor& InputDesc, int32 Rank);
 
-		void SetStridesFromFTensor(const NNECore::Internal::FTensor& InputDesc);
+		void SetStridesFromTensor(const NNECore::Internal::FTensor& InputDesc);
 
-		void UpdateShapesAndStrides(TConstArrayView<uint32> InSizes, TConstArrayView<uint32> InStrides);
+		void UpdateShapeAndStrides(TConstArrayView<uint32> InShape, TConstArrayView<uint32> InStrides = MakeArrayView((uint32*) nullptr, 0));
 
 	private:
 		void Reset();
diff --git a/Engine/Plugins/Experimental/NNE/Source/NNERuntimeRDG/Private/Dml/NNEDmlOperatorInstanceNormalization.cpp b/Engine/Plugins/Experimental/NNE/Source/NNERuntimeRDG/Private/Dml/NNEDmlOperatorInstanceNormalization.cpp
--- a/Engine/Plugins/Experimental/NNE/Source/NNERuntimeRDG/Private/Dml/NNEDmlOperatorInstanceNormalization.cpp
+++ b/Engine/Plugins/Experimental/NNE/Source/NNERuntimeRDG/Private/Dml/NNEDmlOperatorInstanceNormalization.cpp
@@ -6,127 +6,127 @@
 namespace UE::NNERuntimeRDG::Private::Dml
 {
 
-	class FOperatorDmlInstanceNormalization : public FOperatorDml
+class FOperatorDmlInstanceNormalization : public FOperatorDml
+{
+	static constexpr float DefaultEpsilon = 0.00001f;
+
+public:
+
+	static FOperatorDml* Create()
 	{
-		static constexpr float DefaultEpsilon = 0.00001f;
+		return new FOperatorDmlInstanceNormalization();
+	}
 
-	public:
+	//
+	//
+	//
+	virtual bool Initialize(IDMLDevice* Device, TArrayView<const NNECore::Internal::FTensor> InputTensors, TArrayView<const NNECore::Internal::FTensor> OutputTensors, const NNECore::FAttributeMap& Attributes) override
+	{
+		check(InputTensors.Num() >= 1 && InputTensors.Num() <= 3);
+		check(OutputTensors.Num() == 1);
 
-		static FOperatorDml* Create()
-		{
-			return new FOperatorDmlInstanceNormalization();
-		}
+		const NNECore::Internal::FTensor& InputTensor = InputTensors[0];
+		const NNECore::Internal::FTensor& OutputTensor = OutputTensors[0];
 
-		//
-		//
-		//
-		virtual bool Initialize(IDMLDevice* Device, TArrayView<const NNECore::Internal::FTensor> InputTensors, TArrayView<const NNECore::Internal::FTensor> OutputTensors, const NNECore::FAttributeMap& Attributes) override
+		if (InputTensor.GetShape().Rank() > 8)
 		{
-			check(InputTensors.Num() >= 1 && InputTensors.Num() <= 3);
-			check(OutputTensors.Num() == 1);
+			UE_LOG(LogNNE, Warning, TEXT("InstanceNormalization:InputTensor rank should be between 1 and 8, got:%d"), InputTensor.GetShape().Rank());
+			return false;
+		}
 
-			const NNECore::Internal::FTensor& InputTensor = InputTensors[0];
-			const NNECore::Internal::FTensor& OutputTensor = OutputTensors[0];
+		// Read attributes
+		float	Epsilon = Attributes.GetValueOrDefault<float>(TEXT("epsilon"), DefaultEpsilon);
 
-			if (InputTensor.GetShape().Rank() > 8)
-			{
-				UE_LOG(LogNNE, Warning, TEXT("InstanceNormalization:InputTensor rank should be between 1 and 8, got:%d"), InputTensor.GetShape().Rank());
-				return false;
-			}
+		// Initialize tensor descriptors
+		DmlUtil::FTensorDesc	DmlInputTensor{};
+		DmlUtil::FTensorDesc	DmlScalingTensor{};
+		DmlUtil::FTensorDesc	DmlBiasTensor{};
+		DmlUtil::FTensorDesc	DmlOutputTensor{};
+			
+		// Make sure that input is 4D
+		if (!DmlInputTensor.InitFromTensor(InputTensor, 4))
+		{
+			UE_LOG(LogNNE, Warning, TEXT("Failed to initialize tensor(s) for DML inference"));
+			return false;
+		}
 
-			// Read attributes
-			float	Epsilon = Attributes.GetValueOrDefault<float>(TEXT("epsilon"), DefaultEpsilon);
+		if (InputTensors.Num() > 1)
+		{
+			const NNECore::Internal::FTensor& ScaleTensor = InputTensors[1];
 
-			// Initialize tensor descriptors
-			DmlUtil::FTensorDesc	DmlInputTensor{};
-			DmlUtil::FTensorDesc	DmlScalingTensor{};
-			DmlUtil::FTensorDesc	DmlBiasTensor{};
-			DmlUtil::FTensorDesc	DmlOutputTensor{};
-			
-			// Make sure that input is 4D
-			if (!DmlInputTensor.InitFromTensor(InputTensor, 4))
+			if (!DmlScalingTensor.InitFromTensor1D(ScaleTensor, DmlInputTensor.Sizes.Num()))
 			{
 				UE_LOG(LogNNE, Warning, TEXT("Failed to initialize tensor(s) for DML inference"));
 				return false;
 			}
+		}
 
-			if (InputTensors.Num() > 1)
-			{
-				const NNECore::Internal::FTensor& ScaleTensor = InputTensors[1];
-
-				if (!DmlScalingTensor.InitFromTensor1D(ScaleTensor, DmlInputTensor.Sizes.Num()))
-				{
-					UE_LOG(LogNNE, Warning, TEXT("Failed to initialize tensor(s) for DML inference"));
-					return false;
-				}
-			}
-
-			if (InputTensors.Num() > 2)
-			{
-				const NNECore::Internal::FTensor& BiasTensor = InputTensors[2];
-
-				if (!DmlBiasTensor.InitFromTensor1D(BiasTensor, DmlInputTensor.Sizes.Num()))
-				{
-					UE_LOG(LogNNE, Warning, TEXT("Failed to initialize tensor(s) for DML inference"));
-					return false;
-				}
-			}
+		if (InputTensors.Num() > 2)
+		{
+			const NNECore::Internal::FTensor& BiasTensor = InputTensors[2];
 
-			if (!DmlOutputTensor.InitFromTensor(OutputTensor, 4))
+			if (!DmlBiasTensor.InitFromTensor1D(BiasTensor, DmlInputTensor.Sizes.Num()))
 			{
 				UE_LOG(LogNNE, Warning, TEXT("Failed to initialize tensor(s) for DML inference"));
 				return false;
 			}
+		}
+
+		if (!DmlOutputTensor.InitFromTensor(OutputTensor, 4))
+		{
+			UE_LOG(LogNNE, Warning, TEXT("Failed to initialize tensor(s) for DML inference"));
+			return false;
+		}
 
 #ifdef NNE_USE_DML_OPERATOR_MEAN_VARIANCE_NORMALIZATION1
 
-			// "Instance" normalization is really spatial normalization,
-			// where the spatial channels are reduced and normalized, while
-			// batch and channel remain independent. So pass a list of axes
-			// just beyond the leading batch and channel dimensions (starting
-			// at axis 2 up to the last spatial dimension).
-			const int32 InputRank = DmlInputTensor.Sizes.Num();
+		// "Instance" normalization is really spatial normalization,
+		// where the spatial channels are reduced and normalized, while
+		// batch and channel remain independent. So pass a list of axes
+		// just beyond the leading batch and channel dimensions (starting
+		// at axis 2 up to the last spatial dimension).
+		const int32 InputRank = DmlInputTensor.Sizes.Num();
 
-			DmlUtil::FSmallUIntArray Axes;
-			Axes.SetNumUninitialized(NcdhwSpatialDimensionCount);
+		DmlUtil::FSmallUIntArray Axes;
+		Axes.SetNumUninitialized(NcdhwSpatialDimensionCount);
 
-			for (int32 Dim = 0, Value = InputRank - Axes.Num(); Dim < Axes.Num(); ++Dim, ++Value)
-			{
-				Axes[Dim] = Value;
-			}
+		for (int32 Dim = 0, Value = InputRank - Axes.Num(); Dim < Axes.Num(); ++Dim, ++Value)
+		{
+			Axes[Dim] = Value;
+		}
 
-			DML_MEAN_VARIANCE_NORMALIZATION1_OPERATOR_DESC	OpDesc{};
+		DML_MEAN_VARIANCE_NORMALIZATION1_OPERATOR_DESC	OpDesc{};
 
-			OpDesc.InputTensor = &DmlInputTensor.Desc;
-			OpDesc.ScaleTensor = InputTensors.Num() > 1 ? &DmlScalingTensor.Desc : nullptr;
-			OpDesc.BiasTensor = InputTensors.Num() > 2 ? &DmlBiasTensor.Desc : nullptr;
-			OpDesc.OutputTensor = &DmlOutputTensor.Desc;
-			OpDesc.AxisCount = Axes.Num();
-			OpDesc.Axes = Axes.GetData();
-			OpDesc.NormalizeVariance = true;
-			OpDesc.Epsilon = Epsilon;
+		OpDesc.InputTensor = &DmlInputTensor.Desc;
+		OpDesc.ScaleTensor = InputTensors.Num() > 1 ? &DmlScalingTensor.Desc : nullptr;
+		OpDesc.BiasTensor = InputTensors.Num() > 2 ? &DmlBiasTensor.Desc : nullptr;
+		OpDesc.OutputTensor = &DmlOutputTensor.Desc;
+		OpDesc.AxisCount = Axes.Num();
+		OpDesc.Axes = Axes.GetData();
+		OpDesc.NormalizeVariance = true;
+		OpDesc.Epsilon = Epsilon;
 			
-			return CreateOperator(Device, DML_OPERATOR_DESC{ DML_OPERATOR_MEAN_VARIANCE_NORMALIZATION1, &OpDesc });
+		return CreateOperator(Device, DML_OPERATOR_DESC{ DML_OPERATOR_MEAN_VARIANCE_NORMALIZATION1, &OpDesc });
 #else
 
-			DML_MEAN_VARIANCE_NORMALIZATION_OPERATOR_DESC	OpDesc{};
+		DML_MEAN_VARIANCE_NORMALIZATION_OPERATOR_DESC	OpDesc{};
 
-			OpDesc.InputTensor = &DmlInputTensor.Desc;
-			OpDesc.ScaleTensor = InputTensors.Num() > 1 ? &DmlScalingTensor.Desc : nullptr;
-			OpDesc.BiasTensor = InputTensors.Num() > 2 ? &DmlBiasTensor.Desc : nullptr;
-			OpDesc.OutputTensor = &DmlOutputTensor.Desc;
-			OpDesc.CrossChannel = false;
-			OpDesc.NormalizeVariance = true;
-			OpDesc.Epsilon = Epsilon;
+		OpDesc.InputTensor = &DmlInputTensor.Desc;
+		OpDesc.ScaleTensor = InputTensors.Num() > 1 ? &DmlScalingTensor.Desc : nullptr;
+		OpDesc.BiasTensor = InputTensors.Num() > 2 ? &DmlBiasTensor.Desc : nullptr;
+		OpDesc.OutputTensor = &DmlOutputTensor.Desc;
+		OpDesc.CrossChannel = false;
+		OpDesc.NormalizeVariance = true;
+		OpDesc.Epsilon = Epsilon;
 
-			return CreateOperator(Device, DML_OPERATOR_DESC{ DML_OPERATOR_MEAN_VARIANCE_NORMALIZATION, &OpDesc });
+		return CreateOperator(Device, DML_OPERATOR_DESC{ DML_OPERATOR_MEAN_VARIANCE_NORMALIZATION, &OpDesc });
 #endif
 
-		}
-	};
+	}
+};
 
-	// Register operator on Module startup
-	NNE_DML_REGISTER_OP(InstanceNormalization)
+// Register operator on Module startup
+NNE_DML_REGISTER_OP(InstanceNormalization)
 
 } // namespace UE::NNERuntimeRDG::Private::Dml
 
diff --git a/Engine/Plugins/Experimental/NNE/Source/NNERuntimeRDG/Private/Dml/NNEDmlOperatorTranspose.cpp b/Engine/Plugins/Experimental/NNE/Source/NNERuntimeRDG/Private/Dml/NNEDmlOperatorTranspose.cpp
--- a/Engine/Plugins/Experimental/NNE/Source/NNERuntimeRDG/Private/Dml/NNEDmlOperatorTranspose.cpp
+++ b/Engine/Plugins/Experimental/NNE/Source/NNERuntimeRDG/Private/Dml/NNEDmlOperatorTranspose.cpp
@@ -54,7 +54,7 @@ public:
             return false;
         }
         
-        DmlInputTensorDesc.SetStridesFromFTensor(InputTensors[0]);
+        DmlInputTensorDesc.SetStridesFromTensor(InputTensors[0]);
 
         // Initialize Output tensor desc
         DmlUtil::FTensorDesc DmlOutputTensorDesc;
@@ -75,7 +75,7 @@ public:
                 return Permuted;
             };
         
-        DmlInputTensorDesc.UpdateShapesAndStrides(PermuteFunc(DmlInputTensorDesc.Sizes), PermuteFunc(DmlInputTensorDesc.Strides));
+        DmlInputTensorDesc.UpdateShapeAndStrides(PermuteFunc(DmlInputTensorDesc.Sizes), PermuteFunc(DmlInputTensorDesc.Strides));
         
         check(DmlOutputTensorDesc.Sizes == DmlInputTensorDesc.Sizes);        
 
diff --git a/Engine/Plugins/Experimental/NNE/Source/NNERuntimeRDG/Private/Dml/NNEDmlOperatorUpsample.cpp b/Engine/Plugins/Experimental/NNE/Source/NNERuntimeRDG/Private/Dml/NNEDmlOperatorUpsample.cpp
new file mode 100644
--- /dev/null
+++ b/Engine/Plugins/Experimental/NNE/Source/NNERuntimeRDG/Private/Dml/NNEDmlOperatorUpsample.cpp
@@ -0,0 +1,146 @@
+// Copyright Epic Games, Inc. All Rights Reserved.
+
+#ifdef NNE_USE_DIRECTML
+#include "NNEDmlOperator.h"
+
+namespace UE::NNERuntimeRDG::Private::Dml
+{
+
+// Remove array entries of the given indices (in ascending order), shifting them toward the front.
+// There is a special check to avoid removing all the values, since returning a completely
+// empty array would frequently causes errors later in many uses (such as with dimensions).
+//
+// e.g. input values = {2,1,3,1,1,5}
+//      ellidable input indices = {1,3,4}
+//      output values = {2,3,5}
+template< typename TData, typename TAllocator >
+void RemoveValuesByIndex(TConstArrayView<uint32> Indices, TArray<TData, TAllocator>& Values, bool bKeepOneValue)
+{
+	// Keep the last value at least, if all values would otherwise be removed.
+	if (bKeepOneValue && !Indices.IsEmpty() && Indices.Num() == Values.Num()) 
+	{
+		Indices = Indices.RightChop(1);
+	}
+
+	for (int32 Idx = Indices.Num() - 1; Idx >= 0; --Idx)
+	{
+		Values.RemoveAt(Indices[Idx]);
+	}
+}
+
+// Upsample operator is implemented as a DML Resample operator
+class FOperatorDmlUpsample : public FOperatorDml
+{
+	static DML_INTERPOLATION_MODE ModeFromString(FStringView StringVal)
+	{
+		if (FCString::Stricmp(StringVal.GetData(), TEXT("NEAREST")) == 0)
+		{
+			return DML_INTERPOLATION_MODE_NEAREST_NEIGHBOR;
+		}
+		else if (FCString::Stricmp(StringVal.GetData(), TEXT("LINEAR")) == 0)
+		{
+			return DML_INTERPOLATION_MODE_LINEAR;
+		}
+		else
+		{
+			UE_LOG(LogNNE, Warning, TEXT("Unsupported interpolation mode:%s, using nearest neighbor instead"), StringVal.GetData());
+			return DML_INTERPOLATION_MODE_NEAREST_NEIGHBOR;
+		}
+	}
+
+public:
+
+	static FOperatorDml* Create()
+	{
+		return new FOperatorDmlUpsample();
+	}
+
+	//
+	//
+	//
+	virtual bool Initialize(IDMLDevice* Device, TArrayView<const NNECore::Internal::FTensor> InputTensors, TArrayView<const NNECore::Internal::FTensor> OutputTensors, const NNECore::FAttributeMap& Attributes) override
+	{
+		check(InputTensors.Num() == 2);
+		check(OutputTensors.Num() == 1);
+
+		const NNECore::Internal::FTensor& InputTensor = InputTensors[0];
+		const NNECore::Internal::FTensor& ScaleTensor = InputTensors[1];
+		const NNECore::Internal::FTensor& OutputTensor = OutputTensors[0];
+		
+		// Read attributes
+		DML_INTERPOLATION_MODE Mode = ModeFromString(Attributes.GetValue<FString>("mode"));
+
+		// Initialize tensor descriptors
+		DmlUtil::FTensorDesc	DmlInputTensor{};
+		DmlUtil::FTensorDesc	DmlScaleTensor{};
+		DmlUtil::FTensorDesc	DmlOutputTensor{};
+
+		if (!DmlInputTensor.InitFromTensor(InputTensor, 1))
+		{
+			UE_LOG(LogNNE, Warning, TEXT("Failed to initialize tensor(s) for DML inference"));
+			return false;
+		}
+
+		if (!DmlScaleTensor.InitFromTensor(ScaleTensor, 1))
+		{
+			UE_LOG(LogNNE, Warning, TEXT("Failed to initialize tensor(s) for DML inference"));
+			return false;
+		}
+
+		if (!DmlOutputTensor.InitFromTensor(OutputTensor, 1))
+		{
+			UE_LOG(LogNNE, Warning, TEXT("Failed to initialize tensor(s) for DML inference"));
+			return false;
+		}
+
+		// Notify DML graph that we have constant CPU tensor inputs
+		ConstantCPUInputs.Add(1);
+
+		TConstArrayView<float> Scales = ScaleTensor.GetPreparedData<float>();
+		
+		// Find any useless dimensions of size 1 that occur in both input and output
+		DmlUtil::FSmallUIntArray	SqueezeInds;
+		TConstArrayView<uint32>		InputShape = InputTensor.GetShape().GetData();
+		TConstArrayView<uint32>		OutputShape = OutputTensor.GetShape().GetData();
+
+		for (int32 Idx = 0, Rank = OutputTensor.GetShape().Rank(); Idx < Rank; ++Idx)
+		{
+			if (InputShape[Idx] == 1 && OutputShape[Idx] == 1)
+			{
+				SqueezeInds.Emplace(Idx);
+			}
+		}
+
+		if (!SqueezeInds.IsEmpty())
+		{
+			DmlUtil::FSmallUIntArray	SqueezedInputShape(InputShape.GetData(), InputShape.Num());
+			DmlUtil::FSmallUIntArray	SqueezedOutputShape(OutputShape.GetData(), OutputShape.Num());
+			DmlUtil::FSmallArray<float> ScaleValues(Scales.GetData(), Scales.Num());
+
+			RemoveValuesByIndex(SqueezeInds, SqueezedInputShape, true);
+			RemoveValuesByIndex(SqueezeInds, SqueezedOutputShape, true);
+			RemoveValuesByIndex(SqueezeInds, ScaleValues, true);
+
+			DmlInputTensor.UpdateShapeAndStrides(SqueezedInputShape);
+			DmlOutputTensor.UpdateShapeAndStrides(SqueezedOutputShape);
+			Scales = ScaleValues;
+		}
+
+		DML_RESAMPLE_OPERATOR_DESC	OpDesc{};
+
+		OpDesc.InputTensor = &DmlInputTensor.Desc;
+		OpDesc.OutputTensor = &DmlOutputTensor.Desc;
+		OpDesc.InterpolationMode = Mode;
+		OpDesc.ScaleCount = Scales.Num();
+		OpDesc.Scales = Scales.GetData();
+
+		return CreateOperator(Device, DML_OPERATOR_DESC { DML_OPERATOR_RESAMPLE, &OpDesc });
+	}
+};
+
+// Register operator on Module startup
+NNE_DML_REGISTER_OP(Upsample)
+
+} // namespace UE::NNERuntimeRDG::Private::Dml
+
+#endif // NNE_USE_DIRECTML
