NNE - Fix static analysis warning in Dml MatMul

#rnx

diff --git a/Engine/Plugins/Experimental/NNE/Source/NNERuntimeRDG/Private/Dml/NNEDmlOperatorMatMul.cpp b/Engine/Plugins/Experimental/NNE/Source/NNERuntimeRDG/Private/Dml/NNEDmlOperatorMatMul.cpp
--- a/Engine/Plugins/Experimental/NNE/Source/NNERuntimeRDG/Private/Dml/NNEDmlOperatorMatMul.cpp
+++ b/Engine/Plugins/Experimental/NNE/Source/NNERuntimeRDG/Private/Dml/NNEDmlOperatorMatMul.cpp
@@ -102,12 +102,12 @@ public:
 					}
 				}
 
-				if(AIdx < 0 && ActualShapeB[BIdx] != 1)
+				if(AIdx < 0 && BIdx >= 0 && ActualShapeB[BIdx] != 1)
 				{
 					bARequiresBroadcasting = true;
 				}
 
-				if(BIdx < 0 && ActualShapeA[AIdx] != 1)
+				if(BIdx < 0 && AIdx >= 0 && ActualShapeA[AIdx] != 1)
 				{
 					bBRequiresBroadcasting = true;
 				}
