NNE Style transfer dyn shape inference on HLSL

Problem: Style transfer did not worked with variable input shape model.

Solution:
Most of the work was submitted in previous incremental CLs.
This one wrap things up:
* Support scalar tensor for Gather shape inference.
* Fix Cast operator by converting from Protobuf datatype to NNE datatype when building model.
* Fix for RDG model complaining about already upload intermediate tensor when they are const.
* Fix Elem wise operator to support int32/64 on constant tensors.
* Fix Unsqueeze operator fetching wrong attribute name in init.
* Improve logging of tensor type (name instead of enum value).

Tests:
* Remove resnet from tests (unrelated to this CL, just a cleanup)
* Add style transers with dynamic shape tests + model
* Cleanup test definition file
* All test green on win64 editor appart known WIP ops for RDGDml
* Style transfer demo working in editor

#rnx

diff --git a/Engine/Plugins/Experimental/NNE/Source/NNEQA/Resources/NNEQATestDesc.json b/Engine/Plugins/Experimental/NNE/Source/NNEQA/Resources/NNEQATestDesc.json
--- a/Engine/Plugins/Experimental/NNE/Source/NNEQA/Resources/NNEQATestDesc.json
+++ b/Engine/Plugins/Experimental/NNE/Source/NNEQA/Resources/NNEQATestDesc.json
@@ -2,28 +2,29 @@
 	"$schema": "./NNEQATestDescSchema.json",
 	"model_test_categories": [
 		{
-			"category": "OnnxZoo",
+			"category": "StyleTransfer.FixInputSize",
 			"is_model_test": true,
 			"targets": [
 				{
-					"target": "resnet50-v2-7",
-					"tags": [ "SmokeTest" ],
+					"target": "Mosaic-9",
 					"runtimes": [
-						{ "name": "NNERuntimeORTDml" },
-						{ "name": "NNERuntimeRDGHlsl",  "skip": true, "comment": "Missing BatchNormalisation operator." },
 						{ "name": "NNERuntimeRDGDml",  "skip": true, "comment": "Model not supported yet." }
 					],
-					"absolute_tolerance": 1e-5
-				},
+					"absolute_tolerance": 5.0e-2
+				}
+			]
+		},
+		{
+			"category": "StyleTransfer.VariableInputSize",
+			"is_model_test": true,
+			"targets": [
 				{
-					"target": "Mosaic-9",
+					"target": "Mosaic-9-dyn",
 					"runtimes": [
-						{ "name": "NNERuntimeORTDml" },
-						{ "name": "NNERuntimeORTCpu" },
-						{ "name": "NNERuntimeRDGHlsl" },
 						{ "name": "NNERuntimeRDGDml",  "skip": true, "comment": "Model not supported yet." }
 					],
-					"absolute_tolerance": 5.0e-2
+					"absolute_tolerance": 5.0e-2,
+					"relative_tolerance": 5.0e-2
 				}
 			]
 		},
@@ -35,9 +36,7 @@
 					"target": "NeuralMorphModel_global",
 					"tags": [ "SmokeTest" ],
 					"runtimes": [
-						{ "name": "NNERuntimeORTDml" },
-						{ "name": "NNERuntimeRDGHlsl" },
-						{ "name": "NNERuntimeRDGDml" }
+						{ "name": "NNERuntimeRDGDml",  "skip": true, "comment": "Model not supported yet." }
 					],
 					"absolute_tolerance": 1e-5
 				},
@@ -294,22 +293,12 @@
 		},
 		{
 			"category": "Convolution",
-			"runtimes": [
-				{ "name": "NNERuntimeRDGDml" },
-				{ "name": "NNERuntimeORTDml" },
-				{ "name": "NNERuntimeRDGHlsl" }
-			],
 			"targets": [
 				{ "target": "Conv", "absolute_tolerance": 1e-6 }
 			]
 		},
 		{
 			"category": "ConvTranspose",
-			"runtimes": [
-				{ "name": "NNERuntimeRDGDml" },
-				{ "name": "NNERuntimeORTDml" },
-				{ "name": "NNERuntimeRDGHlsl" }
-			],
 			"targets": [
 				{ "target": "ConvTranspose", "absolute_tolerance": 1e-6 }
 			]
@@ -386,6 +375,14 @@
 				{ "inputs": [ { "shape": [ 2,3 ] } ] }
 			]
 		},
+		{
+			"name": "StyleTransfer.VariableInputSize",
+			"datasets": [
+				{ "inputs": [ { "shape": [ 1,3,224,224 ] } ] },
+				{ "inputs": [ { "shape": [ 1,3,128,64 ] } ] },
+				{ "inputs": [ { "shape": [ 1,3,720,1280] } ] }
+			]
+		},
 		{
 			"name": "Elementwise.Unary.Regular",
 			"datasets": [
@@ -608,7 +605,6 @@
 			"datasets": [ { "inputs": [ { "shape": [ 1, 1, 3, 3 ] }, { "shape": [ 1, 2, 3, 3 ] } ], "outputs": [ { "shape": [ 1, 2, 6, 6 ] } ] } ],
 			"runtimes": [
 				{ "name": "NNERuntimeORTDml", "skip": true, "comment": "TODO INVESTIGATE Returns 'The parameter is incorrect'" },
-				{ "name": "NNERuntimeRDGHlsl" }
 			]
 		},
 		{
@@ -630,7 +626,6 @@
 			"name": "MatMul.Rank5",
 			"runtimes": [	
 				{ "name": "NNERuntimeORTDml", "skip": true, "comment": "DML metacommand only support 2 to 4 for inputs rank" },
-				{ "name": "NNERuntimeRDGHlsl"}
 			],
 			"datasets": [
 				{ "inputs": [ { "shape": [ 2, 2, 3, 5, 4 ] }, { "shape": [ 2, 1, 3, 4, 3 ] } ], "outputs": [ { "shape": [ 2, 2, 3, 5, 3 ] } ], "runtimes": [ {"name": "NNERuntimeRDGDml", "skip": true, "comment": "Not Supported"} ]},
@@ -730,74 +725,30 @@
 					"inputs": [ { "shape": [ 2, 2, 4, 1, 3 ] }, { "shape": [ 2 ] }, { "shape": [ 2 ] } ], "outputs": [ { "shape": [ 2, 2, 4, 1, 3 ] } ],
 					"runtimes": [
 						{ "name": "NNERuntimeRDGDml", "skip": true, "comment": "Not supported" },
-						{ "name": "NNERuntimeRDGHlsl" },
-						{ "name": "NNERuntimeORTCpu" },
 						{ "name": "NNERuntimeORTDml", "skip": true, "comment": "Not supported" }
 					]
 				},
 				{
 					"inputs": [ { "shape": [ 2, 2, 4, 4 ] }, { "shape": [ 2 ] }, { "shape": [ 2 ] } ], "outputs": [ { "shape": [ 2, 2, 4, 4 ] } ],
-					"runtimes": [
-						{ "name": "NNERuntimeRDGDml" },
-						{ "name": "NNERuntimeRDGHlsl" },
-						{ "name": "NNERuntimeORTCpu" },
-						{ "name": "NNERuntimeORTDml" }
-					]
 				},
 				{
 					"inputs": [ { "shape": [ 2, 2, 8, 8 ] }, { "shape": [ 2 ] }, { "shape": [ 2 ] } ], "outputs": [ { "shape": [ 2, 2, 8, 8 ] } ],
-					"runtimes": [
-						{ "name": "NNERuntimeRDGDml" },
-						{ "name": "NNERuntimeRDGHlsl" },
-						{ "name": "NNERuntimeORTCpu" },
-						{ "name": "NNERuntimeORTDml" }
-					]
 				},
 				{
 					"inputs": [ { "shape": [ 2, 128, 8, 8 ] }, { "shape": [ 128 ] }, { "shape": [ 128 ] } ], "outputs": [ { "shape": [ 2, 128, 8, 8 ] } ],
-					"runtimes": [
-						{ "name": "NNERuntimeRDGDml" },
-						{ "name": "NNERuntimeRDGHlsl" },
-						{ "name": "NNERuntimeORTCpu" },
-						{ "name": "NNERuntimeORTDml" }
-					]
 				},
 				{
 					"inputs": [ { "shape": [ 128, 2, 8, 8 ] }, { "shape": [ 2 ] }, { "shape": [ 2 ] } ], "outputs": [ { "shape": [ 128, 2, 8, 8 ] } ],
-					"runtimes": [
-						{ "name": "NNERuntimeRDGDml" },
-						{ "name": "NNERuntimeRDGHlsl" },
-						{ "name": "NNERuntimeORTCpu" },
-						{ "name": "NNERuntimeORTDml" }
-					]
 				},
 				{
 					"inputs": [ { "shape": [ 64, 64, 8, 8 ] }, { "shape": [ 64 ] }, { "shape": [ 64 ] } ], "outputs": [ { "shape": [ 64, 64, 8, 8 ] } ],
-					"runtimes": [
-						{ "name": "NNERuntimeRDGDml" },
-						{ "name": "NNERuntimeRDGHlsl" },
-						{ "name": "NNERuntimeORTCpu" },
-						{ "name": "NNERuntimeORTDml" }
-					]
 				},
 				{
 					"inputs": [ { "shape": [ 2, 2, 45, 45 ] }, { "shape": [ 2 ] }, { "shape": [ 2 ] } ], "outputs": [ { "shape": [ 2, 2, 45, 45 ] } ],
-					"runtimes": [
-						{ "name": "NNERuntimeRDGDml" },
-						{ "name": "NNERuntimeRDGHlsl" },
-						{ "name": "NNERuntimeORTCpu" },
-						{ "name": "NNERuntimeORTDml" }
-					]
 				},
 				{
 					"inputs": [ { "shape": [ 1, 128, 56, 56 ] }, { "shape": [ 128 ] }, { "shape": [ 128 ] } ], "outputs": [ { "shape": [ 1, 128, 56, 56 ] } ],
-					"runtimes": [
-						{ "name": "NNERuntimeRDGDml" },
-						{ "name": "NNERuntimeRDGHlsl" },
-						{ "name": "NNERuntimeORTCpu" },
-						{ "name": "NNERuntimeORTDml" }
-					],
-					"comment": "from 'mosaic-9-sim.onnx'"
+					"comment": "from style transfer @ 224x224"
 				}
 			]
 		},
diff --git a/Engine/Plugins/Experimental/NNE/Source/NNERuntimeRDG/Private/Dml/NNEDmlOperatorConvolution.cpp b/Engine/Plugins/Experimental/NNE/Source/NNERuntimeRDG/Private/Dml/NNEDmlOperatorConvolution.cpp
--- a/Engine/Plugins/Experimental/NNE/Source/NNERuntimeRDG/Private/Dml/NNEDmlOperatorConvolution.cpp
+++ b/Engine/Plugins/Experimental/NNE/Source/NNERuntimeRDG/Private/Dml/NNEDmlOperatorConvolution.cpp
@@ -138,7 +138,7 @@ class FOperatorDmlConv : public FOperatorDml
 				SetAutoPadding(InputShape.GetData(), FilterShape.GetData());
 			}
 
-			Group = Attributes.GetValueOrDefault<int>(TEXT("group"), 1);
+			Group = Attributes.GetValueOrDefault<int32>(TEXT("group"), 1);
 
 			if (Direction == DML_CONVOLUTION_DIRECTION_FORWARD)
 			{
diff --git a/Engine/Plugins/Experimental/NNE/Source/NNERuntimeRDG/Private/Dml/NNEDmlOperatorGRU.cpp b/Engine/Plugins/Experimental/NNE/Source/NNERuntimeRDG/Private/Dml/NNEDmlOperatorGRU.cpp
--- a/Engine/Plugins/Experimental/NNE/Source/NNERuntimeRDG/Private/Dml/NNEDmlOperatorGRU.cpp
+++ b/Engine/Plugins/Experimental/NNE/Source/NNERuntimeRDG/Private/Dml/NNEDmlOperatorGRU.cpp
@@ -260,10 +260,10 @@ public:
 		}
 		
 		
-		const int HiddenSizeParam = Attributes.GetValueOrDefault<int>(TEXT("hidden_size"), HiddenSize);
+		const int HiddenSizeParam = Attributes.GetValueOrDefault<int32>(TEXT("hidden_size"), HiddenSize);
 		check(HiddenSize == HiddenSizeParam);
 
-		const int Layout = Attributes.GetValueOrDefault<int>(TEXT("layout"), 0);
+		const int Layout = Attributes.GetValueOrDefault<int32>(TEXT("layout"), 0);
 		if(Layout != 0)
 		{
 			UE_LOG(LogNNE, Error, TEXT("GRU's layout attribute should be 0 for DML inference"));
@@ -296,7 +296,7 @@ public:
 		DmlGRUOpDesc.ActivationDescs = ActivationOpDescs.GetData();
 		
 		DmlGRUOpDesc.LinearBeforeReset = 
-			Attributes.GetValueOrDefault<int>(TEXT("linear_before_reset"), 0) == 0 ? false : true;
+			Attributes.GetValueOrDefault<int32>(TEXT("linear_before_reset"), 0) == 0 ? false : true;
 
 		DML_OPERATOR_DESC DmlOpDesc{};
 
diff --git a/Engine/Plugins/Experimental/NNE/Source/NNERuntimeRDG/Private/Helper/NNERuntimeRDGHelperGather.cpp b/Engine/Plugins/Experimental/NNE/Source/NNERuntimeRDG/Private/Helper/NNERuntimeRDGHelperGather.cpp
--- a/Engine/Plugins/Experimental/NNE/Source/NNERuntimeRDG/Private/Helper/NNERuntimeRDGHelperGather.cpp
+++ b/Engine/Plugins/Experimental/NNE/Source/NNERuntimeRDG/Private/Helper/NNERuntimeRDGHelperGather.cpp
@@ -11,7 +11,7 @@ namespace UE::NNERuntimeRDG::Internal::CPUHelper::Gather
 		TConstArrayView<TInd> IndicesData = IndicesTensor.GetPreparedData<TInd>();
 		TArray<TData> OutputData;
 
-		check(IndicesData.Num() == OutputTensor.GetShape().GetData()[0]);
+		check(IndicesData.Num() == OutputTensor.GetVolume());
 		OutputData.Reserve(IndicesData.Num());
 
 		for (TInd indices : IndicesData)
@@ -66,7 +66,7 @@ namespace UE::NNERuntimeRDG::Internal::CPUHelper::Gather
 			return;
 		}
 
-		if (DataTensor.GetShape().Rank() != 1 || IndicesTensor.GetShape().Rank() != 1)
+		if (DataTensor.GetShape().Rank() > 1 || IndicesTensor.GetShape().Rank() > 1)
 		{
 			return;
 		}
diff --git a/Engine/Plugins/Experimental/NNE/Source/NNERuntimeRDG/Private/Hlsl/NNERuntimeRDGCast.cpp b/Engine/Plugins/Experimental/NNE/Source/NNERuntimeRDG/Private/Hlsl/NNERuntimeRDGCast.cpp
--- a/Engine/Plugins/Experimental/NNE/Source/NNERuntimeRDG/Private/Hlsl/NNERuntimeRDGCast.cpp
+++ b/Engine/Plugins/Experimental/NNE/Source/NNERuntimeRDG/Private/Hlsl/NNERuntimeRDGCast.cpp
@@ -4,6 +4,7 @@
 #include "NNERuntimeRDGHelperCast.h"
 #include "NNECoreTensor.h"
 #include "NNECoreTypes.h"
+#include "NNEUtilsLogHelper.h"
 
 namespace UE::NNERuntimeRDG::Private::Hlsl
 {
@@ -41,7 +42,7 @@ namespace UE::NNERuntimeRDG::Private::Hlsl
 			check(InputTensorDescs.Num() == 1);
 			check(OutputTensorDescs.Num() == 1);
 
-			ENNETensorDataType ToFromAttribute = (ENNETensorDataType)Attributes.GetValue<int>(TEXT("to"));
+			ENNETensorDataType ToFromAttribute = (ENNETensorDataType)Attributes.GetValue<int32>(TEXT("to"));
 			ENNETensorDataType ToFromTensor = OutputTensorDescs[0].GetDataType();
 
 			if (ToFromAttribute != ToFromTensor)
@@ -71,76 +72,18 @@ namespace UE::NNERuntimeRDG::Private::Hlsl
 
 		if (bIsValid)
 		{
-			//In ONNX "To" is assumed to be DataType from TensorProto. Here it is a ENNETensorDataType however no conversion was needed as both enum match.
-			//See ENNETensorDataType GetDataTypeFromGraphTensor(Ort::GraphTensorDataType TensorDataType) in NNEUtilsModelOptimizer.cpp
-			//If a conversion is needed in the future it should be done when converting the model from ONNX to RDG format in the model builder.
-			ENNETensorDataType To = (ENNETensorDataType)AttributeMap.GetValue<int>(TEXT("to"));
+			ENNETensorDataType To = (ENNETensorDataType)AttributeMap.GetValue<int32>(TEXT("to"));
 			switch (To)
 			{
-				case ENNETensorDataType::None:
-					UE_LOG(LogNNE, Warning, TEXT("Cast: Invalid target type 'None' not supported."));
-					bIsValid = false;
-					break;
-				case ENNETensorDataType::Char:
-					UE_LOG(LogNNE, Warning, TEXT("Cast: Invalid target type 'Char' not supported."));
-					bIsValid = false;
-					break;
-				case ENNETensorDataType::Boolean:
-					UE_LOG(LogNNE, Warning, TEXT("Cast: Invalid target type 'Boolean' not supported."));
-					bIsValid = false;
-					break;
-				case ENNETensorDataType::Half:
-					UE_LOG(LogNNE, Warning, TEXT("Cast: Invalid target type 'Half' not supported."));
-					bIsValid = false;
-					break;
 				case ENNETensorDataType::Float:
 					break;
-				case ENNETensorDataType::Double:
-					UE_LOG(LogNNE, Warning, TEXT("Cast: Invalid target type 'Double' not supported."));
-					bIsValid = false;
-					break;
-				case ENNETensorDataType::Int8:
-					UE_LOG(LogNNE, Warning, TEXT("Cast: Invalid target type 'Int8' not supported."));
-					bIsValid = false;
-					break;
-				case ENNETensorDataType::Int16:
-					UE_LOG(LogNNE, Warning, TEXT("Cast: Invalid target type 'Int16' not supported."));
-					bIsValid = false;
-					break;
 				case ENNETensorDataType::Int32:
 					break;
 				case ENNETensorDataType::Int64:
 					break;
-				case ENNETensorDataType::UInt8:
-					UE_LOG(LogNNE, Warning, TEXT("Cast: Invalid target type 'UInt8' not supported."));
-					bIsValid = false;
-					break;
-				case ENNETensorDataType::UInt16:
-					UE_LOG(LogNNE, Warning, TEXT("Cast: Invalid target type 'UInt16' not supported."));
-					bIsValid = false;
-					break;
-				case ENNETensorDataType::UInt32:
-					UE_LOG(LogNNE, Warning, TEXT("Cast: Invalid target type 'UInt32' not supported."));
-					bIsValid = false;
-					break;
-				case ENNETensorDataType::UInt64:
-					UE_LOG(LogNNE, Warning, TEXT("Cast: Invalid target type 'UInt64' not supported."));
-					bIsValid = false;
-					break;
-				case ENNETensorDataType::Complex64:
-					UE_LOG(LogNNE, Warning, TEXT("Cast: Invalid target type 'Complex64' not supported."));
-					bIsValid = false;
-					break;
-				case ENNETensorDataType::Complex128:
-					UE_LOG(LogNNE, Warning, TEXT("Cast: Invalid target type 'Complex128' not supported."));
-					bIsValid = false;
-					break;
-				case ENNETensorDataType::BFloat16:
-					UE_LOG(LogNNE, Warning, TEXT("Cast: Invalid target type 'BFloat16' not supported."));
-					bIsValid = false;
-					break;
 				default:
-					UE_LOG(LogNNE, Warning, TEXT("Cast: Invalid target type %d not supported."), To);
+					FString TargetType = NNEUtils::Internal::GetTensorDataTypeName(To);
+					UE_LOG(LogNNE, Warning, TEXT("Cast: Target tensor data type %s not supported."), *TargetType);
 					bIsValid = false;
 			}
 		}
diff --git a/Engine/Plugins/Experimental/NNE/Source/NNERuntimeRDG/Private/Hlsl/NNERuntimeRDGConv.cpp b/Engine/Plugins/Experimental/NNE/Source/NNERuntimeRDG/Private/Hlsl/NNERuntimeRDGConv.cpp
--- a/Engine/Plugins/Experimental/NNE/Source/NNERuntimeRDG/Private/Hlsl/NNERuntimeRDGConv.cpp
+++ b/Engine/Plugins/Experimental/NNE/Source/NNERuntimeRDG/Private/Hlsl/NNERuntimeRDGConv.cpp
@@ -95,7 +95,7 @@ namespace UE::NNERuntimeRDG::Private::Hlsl
 
 			FConvCS::LexFromString(AutoPad, *Attributes.GetValue<FString>(TEXT("auto_pad")));
 			Dilations = Attributes.GetValueOrDefault<TArray<int32>>(TEXT("dilations"), DilationsOrStridesDefault);
-			Group = Attributes.GetValueOrDefault<int>(TEXT("group"), 1);
+			Group = Attributes.GetValueOrDefault<int32>(TEXT("group"), 1);
 			if (AutoPad == EConvAutoPad::NOTSET)
 			{
 				TArray<int32> PadsDefault;
diff --git a/Engine/Plugins/Experimental/NNE/Source/NNERuntimeRDG/Private/Hlsl/NNERuntimeRDGConvTranspose.cpp b/Engine/Plugins/Experimental/NNE/Source/NNERuntimeRDG/Private/Hlsl/NNERuntimeRDGConvTranspose.cpp
--- a/Engine/Plugins/Experimental/NNE/Source/NNERuntimeRDG/Private/Hlsl/NNERuntimeRDGConvTranspose.cpp
+++ b/Engine/Plugins/Experimental/NNE/Source/NNERuntimeRDG/Private/Hlsl/NNERuntimeRDGConvTranspose.cpp
@@ -95,7 +95,7 @@ namespace UE::NNERuntimeRDG::Private::Hlsl
 
 			FConvTransposeCS::LexFromString(AutoPad, *Attributes.GetValue<FString>(TEXT("auto_pad")));
 			Dilations = Attributes.GetValueOrDefault<TArray<int32>>(TEXT("dilations"), DilationsOrStridesDefault);
-			Group = Attributes.GetValueOrDefault<int>(TEXT("group"), 1);
+			Group = Attributes.GetValueOrDefault<int32>(TEXT("group"), 1);
 			
 			TArray<int32> OutputPaddingDefault;
 			OutputPaddingDefault.Init(0, NumDimensions);
diff --git a/Engine/Plugins/Experimental/NNE/Source/NNERuntimeRDG/Private/Hlsl/NNERuntimeRDGElementWiseBinary.cpp b/Engine/Plugins/Experimental/NNE/Source/NNERuntimeRDG/Private/Hlsl/NNERuntimeRDGElementWiseBinary.cpp
--- a/Engine/Plugins/Experimental/NNE/Source/NNERuntimeRDG/Private/Hlsl/NNERuntimeRDGElementWiseBinary.cpp
+++ b/Engine/Plugins/Experimental/NNE/Source/NNERuntimeRDG/Private/Hlsl/NNERuntimeRDGElementWiseBinary.cpp
@@ -59,6 +59,12 @@ namespace UE::NNERuntimeRDG::Private::Hlsl
 			OutputTensors[0]->SetShape(OutputShape);
 
 			Internal::CPUHelper::ElementWiseBinary::Apply(OpType, *InputTensors[0], *InputTensors[1], *OutputTensors[0]);
+
+			if (OutputTensors[0]->GetDataType() != ENNETensorDataType::Float && !OutputTensors[0]->HasPreparedData())
+			{
+				UE_LOG(LogNNE, Warning, TEXT("Error: binary element wise op output tensor could not be made constant nor it was of float type. Only floats are supported at the moment on the HLSL compute path."));
+				return -1;
+			}
 			
 			return 0;
 		}
@@ -67,6 +73,8 @@ namespace UE::NNERuntimeRDG::Private::Hlsl
 		{
 			check(InputTensorDescs.Num() == 2);
 			check(OutputTensorDescs.Num() == 1);
+			check(InputTensorDescs[0].GetDataType() == InputTensorDescs[1].GetDataType());
+			check(InputTensorDescs[0].GetDataType() == OutputTensorDescs[0].GetDataType());
 		
 			return true;
 		}
@@ -128,6 +136,8 @@ namespace UE::NNERuntimeRDG::Private::Hlsl
 
 		FInputValidator InputValidator;
 		InputValidator.AddSupportedType(ENNETensorDataType::Float);
+		InputValidator.AddSupportedType(ENNETensorDataType::Int32);
+		InputValidator.AddSupportedType(ENNETensorDataType::Int64);
 		InputValidator.AddRequired();
 		InputValidator.AddRequired();
 		bIsValid &= InputValidator.Validate(InputTypes);
diff --git a/Engine/Plugins/Experimental/NNE/Source/NNERuntimeRDG/Private/Hlsl/NNERuntimeRDGGather.cpp b/Engine/Plugins/Experimental/NNE/Source/NNERuntimeRDG/Private/Hlsl/NNERuntimeRDGGather.cpp
--- a/Engine/Plugins/Experimental/NNE/Source/NNERuntimeRDG/Private/Hlsl/NNERuntimeRDGGather.cpp
+++ b/Engine/Plugins/Experimental/NNE/Source/NNERuntimeRDG/Private/Hlsl/NNERuntimeRDGGather.cpp
@@ -33,7 +33,6 @@ namespace UE::NNERuntimeRDG::Private::Hlsl
 		{
 			check(InputTensors.Num() == 2)
 			check(OutputTensors.Num() == 1)
-			UE_LOG(LogNNE, Warning, TEXT("Gather shape inference is not implemented at the moment"));
 
 			const NNECore::Internal::FTensor& Data = *InputTensors[0];
 			const NNECore::Internal::FTensor& Indices = *InputTensors[1];
@@ -85,21 +84,12 @@ namespace UE::NNERuntimeRDG::Private::Hlsl
 			const NNECore::FTensorDesc& Indices = InputTensorDescs[1];
 			const NNECore::FTensorDesc& Output = OutputTensorDescs[0];
 
-			if (Output.GetShape().Rank() <= MaxNumDimensions)
+			if (Output.GetShape().Rank() > MaxNumDimensions)
 			{
-				UE_LOG(LogNNE, Warning, TEXT("Gather first input should be of rank %d or less"), MaxNumDimensions);
-				return false;
-			}
-			if (Data.GetShape().Rank() == 0)
-			{
-				UE_LOG(LogNNE, Warning, TEXT("Gather first input should be at least of rank 1"));
-				return false;
-			}
-			if (Indices.GetShape().Rank() == 0)
-			{
-				UE_LOG(LogNNE, Warning, TEXT("Gather second input should be at least of rank 1"));
+				UE_LOG(LogNNE, Warning, TEXT("Gather first input should be of rank %d or less but is %d"), MaxNumDimensions, Output.GetShape().Rank());
 				return false;
 			}
+
 			if ((Data.GetShape().Rank() + Indices.GetShape().Rank() - 1) > MaxNumDimensions)
 			{
 				UE_LOG(LogNNE, Warning, TEXT("Gather sum of input 0 and 1 ranks -1 should be less than %d"), MaxNumDimensions);
@@ -183,6 +173,8 @@ namespace UE::NNERuntimeRDG::Private::Hlsl
 		InputValidator.SetTemplateCount(2);
 	
 		InputValidator.AddSupportedType(ENNETensorDataType::Float, 0);
+		InputValidator.AddSupportedType(ENNETensorDataType::Int32, 0);
+		InputValidator.AddSupportedType(ENNETensorDataType::Int64, 0);
 		InputValidator.AddRequired(0);
 
 		InputValidator.AddSupportedType(ENNETensorDataType::Int32, 1);
diff --git a/Engine/Plugins/Experimental/NNE/Source/NNERuntimeRDG/Private/Hlsl/NNERuntimeRDGUnsqueeze.cpp b/Engine/Plugins/Experimental/NNE/Source/NNERuntimeRDG/Private/Hlsl/NNERuntimeRDGUnsqueeze.cpp
--- a/Engine/Plugins/Experimental/NNE/Source/NNERuntimeRDG/Private/Hlsl/NNERuntimeRDGUnsqueeze.cpp
+++ b/Engine/Plugins/Experimental/NNE/Source/NNERuntimeRDG/Private/Hlsl/NNERuntimeRDGUnsqueeze.cpp
@@ -52,7 +52,7 @@ namespace UE::NNERuntimeRDG::Private::Hlsl
 			check(InputTensorDescs.Num() == 1);
 			check(OutputTensorDescs.Num() == 1);
 
-			Axes = Attributes.GetValue<TArray<int32>>(TEXT("pads"));
+			Axes = Attributes.GetValue<TArray<int32>>(TEXT("axes"));
 			for (int32 Axe : Axes)
 			{
 				if (Axe < 0)
diff --git a/Engine/Plugins/Experimental/NNE/Source/NNERuntimeRDG/Private/NNERuntimeRDG.cpp b/Engine/Plugins/Experimental/NNE/Source/NNERuntimeRDG/Private/NNERuntimeRDG.cpp
--- a/Engine/Plugins/Experimental/NNE/Source/NNERuntimeRDG/Private/NNERuntimeRDG.cpp
+++ b/Engine/Plugins/Experimental/NNE/Source/NNERuntimeRDG/Private/NNERuntimeRDG.cpp
@@ -3,6 +3,7 @@
 #include "NNERuntimeRDG.h"
 
 #include "NNECoreRuntimeFormat.h"
+#include "NNEUtilsLogHelper.h"
 #include "NNEUtilsModelOptimizer.h"
 #include "RenderGraphBuilder.h"
 #include "RenderGraphUtils.h"
@@ -58,7 +59,8 @@ bool FInputValidator::Validate(TConstArrayView<ENNETensorDataType> InputTypes)
 		check(TemplateIdx < TemplateTypes.Num());
 		if (INDEX_NONE == TemplateTypes[TemplateIdx].Find(InputTypes[Idx]))
 		{
-			UE_LOG(LogNNE, Warning, TEXT("Input '%d' from template idx '%d' is of type '%d' is not supported."), Idx, TemplateIdx, (int)InputTypes[Idx]);
+			FString TargetType = NNEUtils::Internal::GetTensorDataTypeName(InputTypes[Idx]);
+			UE_LOG(LogNNE, Warning, TEXT("Input at index '%d' (from template T%d) is of type '%s' witch is not supported for that input."), Idx, TemplateIdx, *TargetType);
 			bAreInputValid = false;
 		}
 	}
diff --git a/Engine/Plugins/Experimental/NNE/Source/NNERuntimeRDG/Private/NNERuntimeRDGModel.cpp b/Engine/Plugins/Experimental/NNE/Source/NNERuntimeRDG/Private/NNERuntimeRDGModel.cpp
--- a/Engine/Plugins/Experimental/NNE/Source/NNERuntimeRDG/Private/NNERuntimeRDGModel.cpp
+++ b/Engine/Plugins/Experimental/NNE/Source/NNERuntimeRDG/Private/NNERuntimeRDGModel.cpp
@@ -224,13 +224,20 @@ int FModelRDG::EnqueueRDG(FRDGBuilder& RDGBuilder, TConstArrayView<NNECore::FTen
 		return -1;
 	}
 
-	//Create buffer for intermediate tensors
+	//Create buffer for not const intermediate tensors
 	for (FTensorRDG& TensorRDG : IntermediateTensorRDGs)
 	{
-		const FRDGBufferDesc BufferDesc = CreateRDGBufferDescForTensorRDG(TensorRDG);
-		const FRDGBufferRef TensorBuffer = RDGBuilder.CreateBuffer(BufferDesc, *TensorRDG.GetName(), ERDGBufferFlags::None);
-		check(TensorRDG.GetBuffer() == nullptr);
-		TensorRDG.SetBuffer(TensorBuffer);
+		if (!TensorRDG.HasPreparedData())
+		{
+			const FRDGBufferDesc BufferDesc = CreateRDGBufferDescForTensorRDG(TensorRDG);
+			const FRDGBufferRef TensorBuffer = RDGBuilder.CreateBuffer(BufferDesc, *TensorRDG.GetName(), ERDGBufferFlags::None);
+			check(TensorRDG.GetBuffer() == nullptr);
+			TensorRDG.SetBuffer(TensorBuffer);
+		}
+		else
+		{
+			check(TensorRDG.GetBuffer() != nullptr);
+		}
 	}
 
 	if (AddWeightsToRDGGraph(RDGBuilder))
diff --git a/Engine/Plugins/Experimental/NNE/Source/NNEUtils/Internal/NNEUtilsLogHelper.h b/Engine/Plugins/Experimental/NNE/Source/NNEUtils/Internal/NNEUtilsLogHelper.h
new file mode 100644
--- /dev/null
+++ b/Engine/Plugins/Experimental/NNE/Source/NNEUtils/Internal/NNEUtilsLogHelper.h
@@ -0,0 +1,10 @@
+// Copyright Epic Games, Inc. All Rights Reserved.
+#pragma once
+
+#include "NNECoreTypes.h"
+
+namespace UE::NNEUtils::Internal
+{
+	NNEUTILS_API FString GetTensorDataTypeName(ENNETensorDataType DataType);
+} // UE::NNEUtils::Internal
+
diff --git a/Engine/Plugins/Experimental/NNE/Source/NNEUtils/Private/NNEUtilsLogHelper.cpp b/Engine/Plugins/Experimental/NNE/Source/NNEUtils/Private/NNEUtilsLogHelper.cpp
new file mode 100644
--- /dev/null
+++ b/Engine/Plugins/Experimental/NNE/Source/NNEUtils/Private/NNEUtilsLogHelper.cpp
@@ -0,0 +1,33 @@
+// Copyright Epic Games, Inc. All Rights Reserved.
+
+#include "NNEUtilsLogHelper.h"
+
+namespace UE::NNEUtils::Internal
+{
+	FString GetTensorDataTypeName(ENNETensorDataType DataType)
+	{
+		switch (DataType)
+		{
+			case ENNETensorDataType::None:       return TEXT("None");
+			case ENNETensorDataType::Char:       return TEXT("Char");
+			case ENNETensorDataType::Boolean:    return TEXT("Boolean");
+			case ENNETensorDataType::Half:	     return TEXT("Half");
+			case ENNETensorDataType::Float:	     return TEXT("Float");
+			case ENNETensorDataType::Double:     return TEXT("Double");
+			case ENNETensorDataType::Int8:	     return TEXT("Int8");
+			case ENNETensorDataType::Int16:	     return TEXT("Int16");
+			case ENNETensorDataType::Int32:	     return TEXT("Int32");
+			case ENNETensorDataType::Int64:	     return TEXT("Int64");
+			case ENNETensorDataType::UInt8:	     return TEXT("UInt8");
+			case ENNETensorDataType::UInt16:     return TEXT("UInt16");
+			case ENNETensorDataType::UInt32:     return TEXT("UInt32");
+			case ENNETensorDataType::UInt64:     return TEXT("UInt64");
+			case ENNETensorDataType::Complex64:  return TEXT("Complex64");
+			case ENNETensorDataType::Complex128: return TEXT("Complex128");
+			case ENNETensorDataType::BFloat16:   return TEXT("BFloat16");
+			default: return TEXT("Unknown!");
+		}
+	}
+
+} // UE::NNEUtils::Internal
+
diff --git a/Engine/Plugins/Experimental/NNE/Source/NNEUtils/Private/NNEUtilsModelOptimizer.cpp b/Engine/Plugins/Experimental/NNE/Source/NNEUtils/Private/NNEUtilsModelOptimizer.cpp
--- a/Engine/Plugins/Experimental/NNE/Source/NNEUtils/Private/NNEUtilsModelOptimizer.cpp
+++ b/Engine/Plugins/Experimental/NNE/Source/NNEUtils/Private/NNEUtilsModelOptimizer.cpp
@@ -27,6 +27,9 @@ NNE_THIRD_PARTY_INCLUDES_START
 NNE_THIRD_PARTY_INCLUDES_END
 
 
+namespace UE::NNEUtils::Internal
+{
+
 #define Print(Format, ...) UE_LOG(LogNNE, Display, Format, __VA_ARGS__)
 
 class FModelGraphPrinter
@@ -255,16 +258,13 @@ ENNETensorDataType GetDataTypeFromGraphTensor(Ort::GraphTensorDataType TensorDat
 		CASE(kComplex128, Complex128);
 		CASE(kBFloat16, BFloat16);
 
-		default:
-			return ENNETensorDataType::None;
+	default:
+		return ENNETensorDataType::None;
 	}
 }
 
 #undef CASE
 
-namespace UE::NNEUtils::Internal
-{
-
 class FModelOptimizerBase : public NNECore::Internal::IModelOptimizer
 {
 protected:
@@ -489,6 +489,17 @@ public:
 	}
 
 private:
+	int32 ConvertInt32AttributeValueFromORT(int32 ORTValue, const FString& AttributeName, const FString& OpName)
+	{
+		if (OpName == TEXT("Cast") && AttributeName == TEXT("to"))
+		{
+			//Cast.to attribute follow TensorProto.DataType itself a match to Ort::GraphTensorDataType
+			//however NNE format is agnostic to TensorProto thus we need a convertion.
+			return (int32)GetDataTypeFromGraphTensor((Ort::GraphTensorDataType)ORTValue);
+		}
+		return ORTValue;
+	}
+
 	bool BuildNNEFormat(Ort::IModelGraph* Graph, TArray<uint8>& NNEData)
 	{
 		TUniquePtr<IModelBuilder>	Builder(CreateNNEModelBuilder());
@@ -552,7 +563,8 @@ private:
 				}
 				else if (AttrInfo.type == Ort::GraphAttributeType::kInt)
 				{
-					Builder->AddOperatorAttribute(Op, AttrInfo.name, FNNEAttributeValue((int32)AttrValue.i));
+					int32 NNEAttributeValue = ConvertInt32AttributeValueFromORT((int32)AttrValue.i, AttrInfo.name, NodeInfo.opName);
+					Builder->AddOperatorAttribute(Op, AttrInfo.name, FNNEAttributeValue(NNEAttributeValue));
 				}
 				else if (AttrInfo.type == Ort::GraphAttributeType::kInts)
 				{
