NNE : Fix test definition

#rnx

diff --git a/Engine/Plugins/Experimental/NNE/Source/NNEQA/Resources/NNEQATestDesc.json b/Engine/Plugins/Experimental/NNE/Source/NNEQA/Resources/NNEQATestDesc.json
--- a/Engine/Plugins/Experimental/NNE/Source/NNEQA/Resources/NNEQATestDesc.json
+++ b/Engine/Plugins/Experimental/NNE/Source/NNEQA/Resources/NNEQATestDesc.json
@@ -13,8 +13,8 @@
 					"absolute_tolerance": 5.0e-2
 				}
 			]
-		},
-		{
+				},
+				{
 			"category": "StyleTransfer.VariableInputSize",
 			"is_model_test": true,
 			"targets": [
@@ -604,7 +604,7 @@
 			"name": "ConvTranspose.Autopad",
 			"datasets": [ { "inputs": [ { "shape": [ 1, 1, 3, 3 ] }, { "shape": [ 1, 2, 3, 3 ] } ], "outputs": [ { "shape": [ 1, 2, 6, 6 ] } ] } ],
 			"runtimes": [
-				{ "name": "NNERuntimeORTDml", "skip": true, "comment": "TODO INVESTIGATE Returns 'The parameter is incorrect'" },
+				{ "name": "NNERuntimeORTDml", "skip": true, "comment": "TODO INVESTIGATE Returns 'The parameter is incorrect'" }
 			]
 		},
 		{
@@ -625,7 +625,7 @@
 		{
 			"name": "MatMul.Rank5",
 			"runtimes": [	
-				{ "name": "NNERuntimeORTDml", "skip": true, "comment": "DML metacommand only support 2 to 4 for inputs rank" },
+				{ "name": "NNERuntimeORTDml", "skip": true, "comment": "DML metacommand only support 2 to 4 for inputs rank" }
 			],
 			"datasets": [
 				{ "inputs": [ { "shape": [ 2, 2, 3, 5, 4 ] }, { "shape": [ 2, 1, 3, 4, 3 ] } ], "outputs": [ { "shape": [ 2, 2, 3, 5, 3 ] } ], "runtimes": [ {"name": "NNERuntimeRDGDml", "skip": true, "comment": "Not Supported"} ]},
@@ -729,22 +729,22 @@
 					]
 				},
 				{
-					"inputs": [ { "shape": [ 2, 2, 4, 4 ] }, { "shape": [ 2 ] }, { "shape": [ 2 ] } ], "outputs": [ { "shape": [ 2, 2, 4, 4 ] } ],
+					"inputs": [ { "shape": [ 2, 2, 4, 4 ] }, { "shape": [ 2 ] }, { "shape": [ 2 ] } ], "outputs": [ { "shape": [ 2, 2, 4, 4 ] } ]
 				},
 				{
-					"inputs": [ { "shape": [ 2, 2, 8, 8 ] }, { "shape": [ 2 ] }, { "shape": [ 2 ] } ], "outputs": [ { "shape": [ 2, 2, 8, 8 ] } ],
+					"inputs": [ { "shape": [ 2, 2, 8, 8 ] }, { "shape": [ 2 ] }, { "shape": [ 2 ] } ], "outputs": [ { "shape": [ 2, 2, 8, 8 ] } ]
 				},
 				{
-					"inputs": [ { "shape": [ 2, 128, 8, 8 ] }, { "shape": [ 128 ] }, { "shape": [ 128 ] } ], "outputs": [ { "shape": [ 2, 128, 8, 8 ] } ],
+					"inputs": [ { "shape": [ 2, 128, 8, 8 ] }, { "shape": [ 128 ] }, { "shape": [ 128 ] } ], "outputs": [ { "shape": [ 2, 128, 8, 8 ] } ]
 				},
 				{
-					"inputs": [ { "shape": [ 128, 2, 8, 8 ] }, { "shape": [ 2 ] }, { "shape": [ 2 ] } ], "outputs": [ { "shape": [ 128, 2, 8, 8 ] } ],
+					"inputs": [ { "shape": [ 128, 2, 8, 8 ] }, { "shape": [ 2 ] }, { "shape": [ 2 ] } ], "outputs": [ { "shape": [ 128, 2, 8, 8 ] } ]
 				},
 				{
-					"inputs": [ { "shape": [ 64, 64, 8, 8 ] }, { "shape": [ 64 ] }, { "shape": [ 64 ] } ], "outputs": [ { "shape": [ 64, 64, 8, 8 ] } ],
+					"inputs": [ { "shape": [ 64, 64, 8, 8 ] }, { "shape": [ 64 ] }, { "shape": [ 64 ] } ], "outputs": [ { "shape": [ 64, 64, 8, 8 ] } ]
 				},
 				{
-					"inputs": [ { "shape": [ 2, 2, 45, 45 ] }, { "shape": [ 2 ] }, { "shape": [ 2 ] } ], "outputs": [ { "shape": [ 2, 2, 45, 45 ] } ],
+					"inputs": [ { "shape": [ 2, 2, 45, 45 ] }, { "shape": [ 2 ] }, { "shape": [ 2 ] } ], "outputs": [ { "shape": [ 2, 2, 45, 45 ] } ]
 				},
 				{
 					"inputs": [ { "shape": [ 1, 128, 56, 56 ] }, { "shape": [ 128 ] }, { "shape": [ 128 ] } ], "outputs": [ { "shape": [ 1, 128, 56, 56 ] } ],
