NNE Allow to call SetInputTensorShape on the rendering thread

Problem: SetInputTensorShape could not be called from the rendering thread

Solution:
Do not upload constant buffers from shape inference in SetInputTensorShape so it can be called from any thread, but rather upload the constant tensor buffer as part of EnqueueRDG using the provided RDGBuilder context.

Tests:
* All test green on win64 editor appart known WIP ops for RDGDml
* Style transfer demo working in editor, also with dynamic input style transfer network

#rnx

diff --git a/Engine/Plugins/Experimental/NNE/Source/NNECore/Public/NNECoreRuntimeCPU.h b/Engine/Plugins/Experimental/NNE/Source/NNECore/Public/NNECoreRuntimeCPU.h
--- a/Engine/Plugins/Experimental/NNE/Source/NNECore/Public/NNECoreRuntimeCPU.h
+++ b/Engine/Plugins/Experimental/NNE/Source/NNECore/Public/NNECoreRuntimeCPU.h
@@ -74,6 +74,7 @@ public:
 	 *
 	 * The call is mandatory before a model can be run.
 	 * The function will run shape inference and resolve, if possible, the output shapes which can then be accessed by calling GetOutputTensorShapes().
+	 * This is a potentially expensive call and should be called lazily if possible.
 	 *
 	 * @param InInputShapes The input shapes to prepare the model with.
 	 * @return 0 on success or a non-zero number otherwise.
diff --git a/Engine/Plugins/Experimental/NNE/Source/NNECore/Public/NNECoreRuntimeRDG.h b/Engine/Plugins/Experimental/NNE/Source/NNECore/Public/NNECoreRuntimeRDG.h
--- a/Engine/Plugins/Experimental/NNE/Source/NNECore/Public/NNECoreRuntimeRDG.h
+++ b/Engine/Plugins/Experimental/NNE/Source/NNECore/Public/NNECoreRuntimeRDG.h
@@ -70,9 +70,11 @@ public:
 	
 	/**
 	 * Prepare the model to be run with the given input shape.
-	 *
+	 * 
 	 * The call is mandatory before a model can be run.
+	 * This function might be called from the render thread, if not it is up to the caller to ensure thread safety.
 	 * The function will run shape inference and resolve, if possible, the output shapes which can then be accessed by calling GetOutputTensorShapes().
+	 * This is a potentially expensive call and should be called lazily if possible.
 	 *
 	 * @param InInputShapes The input shapes to prepare the model with.
 	 * @return 0 on success or a non-zero number otherwise.
diff --git a/Engine/Plugins/Experimental/NNE/Source/NNERuntimeRDG/Private/NNERuntimeRDGModel.cpp b/Engine/Plugins/Experimental/NNE/Source/NNERuntimeRDG/Private/NNERuntimeRDGModel.cpp
--- a/Engine/Plugins/Experimental/NNE/Source/NNERuntimeRDG/Private/NNERuntimeRDGModel.cpp
+++ b/Engine/Plugins/Experimental/NNE/Source/NNERuntimeRDG/Private/NNERuntimeRDGModel.cpp
@@ -224,7 +224,11 @@ int FModelRDG::EnqueueRDG(FRDGBuilder& RDGBuilder, TConstArrayView<NNECore::FTen
 		return -1;
 	}
 
-	//Create buffer for not const intermediate tensors
+	
+	//Register constant and weights tensors resources to RDG graph, uploading constant tensors if needed
+	bool bBuffersUploadedAndRegisteredToRDGGraph = PrepareModelRDG(RDGBuilder);
+
+	//Create temporary buffers for NOT const intermediate tensors
 	for (FTensorRDG& TensorRDG : IntermediateTensorRDGs)
 	{
 		if (!TensorRDG.HasPreparedData())
@@ -234,23 +238,13 @@ int FModelRDG::EnqueueRDG(FRDGBuilder& RDGBuilder, TConstArrayView<NNECore::FTen
 			check(TensorRDG.GetBuffer() == nullptr);
 			TensorRDG.SetBuffer(TensorBuffer);
 		}
-		else
-		{
-			check(TensorRDG.GetBuffer() != nullptr);
-		}
 	}
 
-	if (AddWeightsToRDGGraph(RDGBuilder))
-	{
-		//Note: DirectML uses RHI buffers instead of RDG buffers
-		//For now weights tensors are not uploaded to GPU thus GetBuffer will return nullptr for them.
-		checkCode(for (const FTensorRDG* TensorRDG : AllTensorRDGRefs) { if (TensorRDG != nullptr) { check(TensorRDG->GetBuffer() != nullptr); } });	
-	}
-
-	//Insert weights tensors
-	for (int32 i = 0; i < WeightTensorIndices.Num(); ++i)
+	//Note: DirectML uses RHI buffers instead of RDG buffers
+	//For now weights tensors are not uploaded to GPU thus GetBuffer will return nullptr for them.
+	if (bBuffersUploadedAndRegisteredToRDGGraph)
 	{
-		AllTensorRDGRefs[WeightTensorIndices[i]] = &WeightTensorRDGs[i];
+		checkCode(for (const FTensorRDG* TensorRDG : AllTensorRDGRefs) { if (TensorRDG != nullptr) { check(TensorRDG->GetBuffer() != nullptr); } });
 	}
 
 	// We can now dispatch operators
diff --git a/Engine/Plugins/Experimental/NNE/Source/NNERuntimeRDG/Private/NNERuntimeRDGModel.h b/Engine/Plugins/Experimental/NNE/Source/NNERuntimeRDG/Private/NNERuntimeRDGModel.h
--- a/Engine/Plugins/Experimental/NNE/Source/NNERuntimeRDG/Private/NNERuntimeRDGModel.h
+++ b/Engine/Plugins/Experimental/NNE/Source/NNERuntimeRDG/Private/NNERuntimeRDGModel.h
@@ -25,7 +25,7 @@ namespace UE::NNERuntimeRDG::Private
 		int SetTensors(FRDGBuilder& GraphBuilder, FTensorRDGArray& InTensorRDGs, TConstArrayView<NNECore::FTensorBindingRDG> InBindings);
 
 		virtual int PrepareTensorShapesAndData() = 0;
-		virtual bool AddWeightsToRDGGraph(FRDGBuilder& RDGBuilder) { return false; }
+		virtual bool PrepareModelRDG(FRDGBuilder& RDGBuilder) { return false; }
 		virtual void AddDispatchOps_RenderThread(FRDGBuilder& GraphBuilder) = 0;
 
 		//Tensor descriptor
diff --git a/Engine/Plugins/Experimental/NNE/Source/NNERuntimeRDG/Private/NNERuntimeRDGModelHlsl.cpp b/Engine/Plugins/Experimental/NNE/Source/NNERuntimeRDG/Private/NNERuntimeRDGModelHlsl.cpp
--- a/Engine/Plugins/Experimental/NNE/Source/NNERuntimeRDG/Private/NNERuntimeRDGModelHlsl.cpp
+++ b/Engine/Plugins/Experimental/NNE/Source/NNERuntimeRDG/Private/NNERuntimeRDGModelHlsl.cpp
@@ -37,10 +37,39 @@ FOperatorHlsl* OpCreate(const FString& OpName, TConstArrayView<NNECore::FTensorD
 
 } // namespace ModelUtils
 
-bool FModel::AddWeightsToRDGGraph(FRDGBuilder& RDGBuilder)
+bool FModel::PrepareModelRDG(FRDGBuilder& RDGBuilder)
 {
-	check(WeightTensorRDGs.Num() == WeightsExternalRDGResources.Num());
+	//Register constant tensors to graph, uploading if needed
+	check(IntermediateTensorRDGs.Num() == ConstantsExternalRDGResources.Num());
+	for (int32 Idx = 0; Idx < ConstantsExternalRDGResources.Num(); ++Idx)
+	{
+		FTensorRDG& Tensor = IntermediateTensorRDGs[Idx];
+		TRefCountPtr<FRDGPooledBuffer>& PooledBuffer = ConstantsExternalRDGResources[Idx];
+
+		if (Tensor.HasPreparedData())
+		{
+			if (!PooledBuffer.IsValid())
+			{
+				FRDGBufferDesc BufferDesc = FRDGBufferDesc::CreateBufferDesc(Tensor.GetElemByteSize(), Tensor.GetVolume());
+				const FRDGBufferRef TransientRDGBuffer = RDGBuilder.CreateBuffer(BufferDesc, *Tensor.GetName(), ERDGBufferFlags::None);
+				const uint8* TensorData = Tensor.GetPreparedData<uint8>().GetData();
+				PooledBuffer = RDGBuilder.ConvertToExternalBuffer(TransientRDGBuffer);
+
+				// Data is copied so model can be released safely or another upload added to the queue
+				RDGBuilder.QueueBufferUpload(TransientRDGBuffer, TensorData, Tensor.GetDataSize(), ERDGInitialDataFlags::None);
+			}
+			check(PooledBuffer.IsValid())
+			FRDGBufferRef Buffer = RDGBuilder.RegisterExternalBuffer(PooledBuffer);
+			Tensor.SetBuffer(Buffer);
+		}
+		else
+		{
+			Tensor.SetBuffer(nullptr);
+		}
+	}
 
+	//Register weight tensors to graph
+	check(WeightTensorRDGs.Num() == WeightsExternalRDGResources.Num());
 	for (int32 Idx = 0; Idx < WeightsExternalRDGResources.Num(); ++Idx)
 	{
 		const TRefCountPtr<FRDGPooledBuffer>& PooledBuffer = WeightsExternalRDGResources[Idx];
@@ -175,6 +204,10 @@ int FModel::PrepareTensorShapesAndData()
 		}
 	);
 
+	//Release uploaded GPU side constants tensors.
+	ConstantsExternalRDGResources.Reset();
+	ConstantsExternalRDGResources.SetNum(IntermediateTensorRDGs.Num());
+
 	// Run model preparation (including shape inference) on all operators
 	// This loop could be abstracted to a different system as it apply on FTensorRef & IPrepareOperator witch are RDG agnostics.
 	for (int32 Idx = 0; Idx < Operators.Num(); ++Idx)
@@ -215,12 +248,6 @@ int FModel::PrepareTensorShapesAndData()
 		};
 	);
 
-	if (!PrepareConstants())
-	{
-		UE_LOG(LogNNE, Warning, TEXT("Constant tensors could not be uploaded to GPU memory."));
-		return -1;
-	}
-
 	return 0;
 }
 
@@ -272,7 +299,6 @@ void EnqueueTensorUpload(TArray<TRefCountPtr<FRDGPooledBuffer>>& OutExternalRDGR
 						const uint8* TensorData = Tensor.GetPreparedData<uint8>().GetData();
 
 						OutExternalRDGResources[i] = RDGBuilder.ConvertToExternalBuffer(TransientRDGBuffer);
-						Tensor.SetBuffer(TransientRDGBuffer);
 						RDGBuilder.QueueBufferUpload(TransientRDGBuffer, TensorData, Tensor.GetDataSize(), CopyDataFlag);
 					}
 				}
@@ -296,14 +322,6 @@ void EnqueueTensorUpload(TArray<TRefCountPtr<FRDGPooledBuffer>>& OutExternalRDGR
 	}
 }
 
-bool FModel::PrepareConstants()
-{
-	// Data is copied so model can be released safely or another upload added to the queue (via SetInputShapes).
-	UploadHelper::EnqueueTensorUpload(ConstantsExternalRDGResources, IntermediateTensorRDGs, ERDGInitialDataFlags::None);
-	
-	return true;
-}
-
 bool FModel::PrepareWeights()
 {
 	check(WeightsExternalRDGResources.IsEmpty());
diff --git a/Engine/Plugins/Experimental/NNE/Source/NNERuntimeRDG/Private/NNERuntimeRDGModelHlsl.h b/Engine/Plugins/Experimental/NNE/Source/NNERuntimeRDG/Private/NNERuntimeRDGModelHlsl.h
--- a/Engine/Plugins/Experimental/NNE/Source/NNERuntimeRDG/Private/NNERuntimeRDGModelHlsl.h
+++ b/Engine/Plugins/Experimental/NNE/Source/NNERuntimeRDG/Private/NNERuntimeRDGModelHlsl.h
@@ -23,11 +23,10 @@ public:
 protected:
 
 	virtual int PrepareTensorShapesAndData() override;
-	virtual bool AddWeightsToRDGGraph(FRDGBuilder& RDGBuilder) override;
+	virtual bool PrepareModelRDG(FRDGBuilder& RDGBuilder) override;
 	virtual void AddDispatchOps_RenderThread(FRDGBuilder& GraphBuilder) override;
 
 	bool PrepareWeights();
-	bool PrepareConstants();
 
 private:
 
