NNE - DML Reshape
#rnx

diff --git a/Engine/Plugins/Experimental/NNE/Source/NNEQA/Resources/NNEQATestDesc.json b/Engine/Plugins/Experimental/NNE/Source/NNEQA/Resources/NNEQATestDesc.json
--- a/Engine/Plugins/Experimental/NNE/Source/NNEQA/Resources/NNEQATestDesc.json
+++ b/Engine/Plugins/Experimental/NNE/Source/NNEQA/Resources/NNEQATestDesc.json
@@ -364,6 +364,15 @@
 			"targets": [
 				{ "target": "Concat" }
 			]
+		},
+		{
+			"category": "Reshape",
+			"runtimes": [
+				{ "name": "NNERuntimeRDGHlsl", "skip": true }
+			],
+			"targets": [
+				{ "target": "Reshape" }
+			]
 		}
 	],
 
@@ -813,6 +822,36 @@
 			"datasets": [
 				{ "inputs": [ { "shape": [ 4, 2, 4 ] }, { "shape": [ 6, 2, 4 ] }, { "shape": [ 1, 2, 4 ] } ], "outputs": [ { "shape": [ 11, 2, 4 ] } ] }
 			]
+		},
+		{ 
+			"name": "Reshape.Reshape234to423",
+			"datasets": [
+				{ "inputs": [ { "shape": [ 2, 3, 4 ] } ], "weights": [ { "shape": [ 3 ], "source": [ "4", "2", "3" ], "type": "int64" } ], "outputs": [ { "shape": [ 4, 2, 3 ] } ] }
+			]
+		},
+		{ 
+			"name": "Reshape.Reshape234to243",
+			"datasets": [
+				{ "inputs": [ { "shape": [ 2, 3, 4 ] } ], "weights": [ { "shape": [ 3 ], "source": [ "2", "4", "3" ], "type": "int64" } ], "outputs": [ { "shape": [ 2, 4, 3 ] } ] }
+			]
+		},
+		{ 
+			"name": "Reshape.ReducedDims",
+			"datasets": [
+				{ "inputs": [ { "shape": [ 2, 3, 4 ] } ], "weights": [ { "shape": [ 2 ], "source": [ "2", "12" ], "type": "int64" } ], "outputs": [ { "shape": [ 2, 12 ] } ] }
+			]
+		},
+		{ 
+			"name": "Reshape.ExtendedDims",
+			"datasets": [
+				{ "inputs": [ { "shape": [ 2, 3, 4 ] } ], "weights": [ { "shape": [ 4 ], "source": [ "2", "3", "2", "2" ], "type": "int64" } ], "outputs": [ { "shape": [ 2, 3, 2, 2 ] } ] }
+			]
+		},
+		{ 
+			"name": "Reshape.OneDim",
+			"datasets": [
+				{ "inputs": [ { "shape": [ 2, 3, 4 ] } ], "weights": [ { "shape": [ 1 ], "source": [ "24" ], "type": "int64" } ], "outputs": [ { "shape": [ 24 ] } ] }
+			]
 		}
 	],
 
diff --git a/Engine/Plugins/Experimental/NNE/Source/NNERuntimeRDG/Private/Dml/NNEDmlOperatorReshape.cpp b/Engine/Plugins/Experimental/NNE/Source/NNERuntimeRDG/Private/Dml/NNEDmlOperatorReshape.cpp
new file mode 100644
--- /dev/null
+++ b/Engine/Plugins/Experimental/NNE/Source/NNERuntimeRDG/Private/Dml/NNEDmlOperatorReshape.cpp
@@ -0,0 +1,107 @@
+// Copyright Epic Games, Inc. All Rights Reserved.
+
+#ifdef NNE_USE_DIRECTML
+
+#include "NNEDmlOperator.h"
+#include "Algo/Transform.h"
+#include "Misc/EnumerateRange.h"
+
+namespace UE::NNERuntimeRDG::Private::Dml
+{
+
+/**
+ * Reshape
+ */
+class FOperatorDmlReshape : public FOperatorDml
+{
+    template<typename DataType>
+    DmlUtil::FSmallUIntArray TensorContentToUIntArray(const NNECore::Internal::FTensor& InputTensor)
+    {
+        DmlUtil::FSmallUIntArray OutArray;
+        Algo::Transform(InputTensor.GetPreparedData<DataType>(), OutArray, [](DataType In){ return (uint32) In; });
+        return OutArray;
+    }
+
+public:
+
+	static FOperatorDml* Create()
+	{
+		return new FOperatorDmlReshape();
+	}
+
+	//
+	//
+	//
+	virtual bool Initialize(IDMLDevice* Device, TArrayView<const NNECore::Internal::FTensor> InputTensors, TArrayView<const NNECore::Internal::FTensor> OutputTensors, const NNECore::FAttributeMap& Attributes) override
+	{
+        check(InputTensors.Num() == 2);
+        check(OutputTensors.Num() == 1);
+
+        ConstantCPUInputs.Add(1);
+
+        // Shape tensor must be constant!
+        check(InputTensors[1].HasPreparedData());
+
+        DmlUtil::FSmallUIntArray ReshapedShape;
+
+        switch(InputTensors[1].GetDataType())
+        {
+        case ENNETensorDataType::Int32:
+            ReshapedShape = TensorContentToUIntArray<int32>(InputTensors[1]);
+            break;
+        case ENNETensorDataType::Int64:
+            ReshapedShape = TensorContentToUIntArray<int64>(InputTensors[1]);
+            break;
+        case ENNETensorDataType::UInt32:
+            ReshapedShape = TensorContentToUIntArray<uint32>(InputTensors[1]);
+            break;
+        default:
+            UE_LOG(LogNNE, Warning, TEXT("Shape tensor has invalid data type"));
+			return false;
+        };
+        
+
+        bool bAllowZero = (bool)
+			( Attributes.GetValueOrDefault<int32>(TEXT("allowzero"), 0) );
+
+        if(!bAllowZero)
+        {
+            for (TEnumerateRef<uint32> Elem : EnumerateRange(ReshapedShape))
+            {
+                check(*Elem != 0 || InputTensors[0].GetShape().Rank() > Elem.GetIndex());
+                *Elem = *Elem == 0 ? InputTensors[0].GetShape().GetData()[Elem.GetIndex()] : *Elem;
+            }
+        }
+
+        DmlUtil::FTensorDesc DmlInputTensorDesc;
+        if (!DmlInputTensorDesc.InitFromTensor(InputTensors[0], ReshapedShape.Num(), 
+			/*Broadcast =*/ MakeArrayView((uint32*) nullptr, 0), 
+			/*CustomShape =*/ ReshapedShape))
+		{
+			UE_LOG(LogNNE, Warning, TEXT("Failed to initialize Reshape's input tensor for DML inference"));
+			return false;
+		}
+
+        DmlUtil::FTensorDesc DmlOutputTensorDesc;
+        if (!DmlOutputTensorDesc.InitFromTensor(OutputTensors[0], OutputTensors[0].GetShape().Rank()))
+		{
+			UE_LOG(LogNNE, Warning, TEXT("Failed to initialize Reshape's output tensor for DML inference"));
+			return false;
+		}
+        
+		DML_ELEMENT_WISE_IDENTITY_OPERATOR_DESC DmlIdentityOpDesc{};
+
+        DmlIdentityOpDesc.InputTensor = &DmlInputTensorDesc.Desc;
+        DmlIdentityOpDesc.OutputTensor = &DmlOutputTensorDesc.Desc;
+
+		return CreateOperator(Device, DML_OPERATOR_DESC{ DML_OPERATOR_ELEMENT_WISE_IDENTITY, &DmlIdentityOpDesc} );
+
+	}
+};
+
+// Register Reshape operator on Module startup
+NNE_DML_REGISTER_OP(Reshape)
+
+} // namespace UE::NNERuntimeRDG::Private::Dml
+
+#endif // NNE_USE_DIRECTML
