NNE Add Squeeze,Size and Reshape to HLSL. Also GPU side of Unsqueeze.

Problem:
Those operator were missing both for shape inference / tensor constness propagation and model inference. Those are common operators especially in the shape inference context.
Also unsqueeze didn't had a GPU side implementation.

Solution:
Add the ops both for shape inference and gpu side (those op gpu side is just a buffer copy)

Tests:
* NNE all test green but known WIP ones + added related op test for new op + Style transfer demo

#rnx

diff --git a/Engine/Plugins/Experimental/NNE/Source/NNEQA/Resources/NNEQATestDesc.json b/Engine/Plugins/Experimental/NNE/Source/NNEQA/Resources/NNEQATestDesc.json
--- a/Engine/Plugins/Experimental/NNE/Source/NNEQA/Resources/NNEQATestDesc.json
+++ b/Engine/Plugins/Experimental/NNE/Source/NNEQA/Resources/NNEQATestDesc.json
@@ -356,6 +356,16 @@
 				{ "target": "Shape" }
 			]
 		},
+		{
+			"category": "SizeOp",
+			"runtimes": [
+				{ "name": "NNERuntimeRDGDml", "skip": true },
+				{ "name": "NNERuntimeRDGHlsl", "skip_static_test": true  }
+			],
+			"targets": [
+				{ "target": "Size" }
+			]
+		},
 		{
 			"category": "Concat",
 			"runtimes": [
@@ -368,11 +378,28 @@
 		{
 			"category": "Reshape",
 			"runtimes": [
-				{ "name": "NNERuntimeRDGHlsl", "skip": true }
 			],
 			"targets": [
 				{ "target": "Reshape" }
 			]
+		},
+		{
+			"category": "SqueezeOp",
+			"runtimes": [
+				{ "name": "NNERuntimeRDGDml", "skip": true },
+			],
+			"targets": [
+				{ "target": "Squeeze" }
+			]
+		},
+		{
+			"category": "Unsqueeze",
+			"runtimes": [
+				{ "name": "NNERuntimeRDGDml", "skip": true },
+			],
+			"targets": [
+				{ "target": "Unsqueeze" }
+			]
 		}
 	],
 
@@ -811,18 +838,38 @@
 				{ "inputs":  [ { "shape": [ 2, 1, 3, 1, 2 ] } ],  "outputs": [ { "shape": [ 5 ], "type": "int64"} ] }
 			]
 		},
+		{
+			"name": "SizeOp",
+			"datasets": [
+				{ "inputs":  [ { "shape": [ 3 ] } ], "outputs": [ { "shape": [ ], "type": "int64" } ] },
+				{ "inputs":  [ { "shape": [ 1, 128, 50, 55 ] } ], "outputs": [ { "shape": [ ], "type": "int64"} ] },
+				{ "inputs":  [ { "shape": [ 2, 1, 3, 1, 2 ] } ],  "outputs": [ { "shape": [ ], "type": "int64"} ] }
+			]
+		},
 		{ 
 			"name": "Concat.Concat1",
 			"datasets": [
 				{ "inputs": [ { "shape": [ 2, 3 ] }, { "shape": [ 2, 4 ] } ], "outputs": [ { "shape": [ 2, 7 ] } ] }
 			]
 		},
-		{ 
+		{
 			"name": "Concat.Concat2",
 			"datasets": [
 				{ "inputs": [ { "shape": [ 4, 2, 4 ] }, { "shape": [ 6, 2, 4 ] }, { "shape": [ 1, 2, 4 ] } ], "outputs": [ { "shape": [ 11, 2, 4 ] } ] }
 			]
 		},
+		{ 
+			"name": "Unsqueeze",
+			"datasets": [
+				{ "inputs": [ { "shape": [ 2, 3, 4 ] } ], "outputs": [ { "shape": [ 2, 1, 3, 4 ] } ] }
+			]
+		},
+		{ 
+			"name": "SqueezeOp",
+			"datasets": [
+				{ "inputs": [ { "shape": [ 2, 3, 1, 4 ] } ], "outputs": [ { "shape": [ 2, 3, 4 ] } ] }
+			]
+		},
 		{ 
 			"name": "Reshape.Reshape234to423",
 			"datasets": [
@@ -1225,6 +1272,22 @@
 				{ "attributes": [] }
 			]
 		},
+		{
+			"name": "SqueezeOp",
+			"attribute_maps": [
+				{ "attributes": [
+					{ "name": "axes", "value": [2], "type": "Int32Array" }
+				] }
+			]
+		},
+		{
+			"name": "Unsqueeze",
+			"attribute_maps": [
+				{ "attributes": [
+					{ "name": "axes", "value": [1], "type": "Int32Array" }
+				] }
+			]
+		},
 		{
 			"name": "Concat.Concat1",
 			"attribute_maps": [
diff --git a/Engine/Plugins/Experimental/NNE/Source/NNERuntimeRDG/Private/Hlsl/NNERuntimeRDGReshape.cpp b/Engine/Plugins/Experimental/NNE/Source/NNERuntimeRDG/Private/Hlsl/NNERuntimeRDGReshape.cpp
new file mode 100644
--- /dev/null
+++ b/Engine/Plugins/Experimental/NNE/Source/NNERuntimeRDG/Private/Hlsl/NNERuntimeRDGReshape.cpp
@@ -0,0 +1,128 @@
+// Copyright Epic Games, Inc. All Rights Reserved.
+
+#include "NNERuntimeRDGReshape.h"
+#include "NNECoreTensor.h"
+#include "NNECoreTypes.h"
+#include "RenderGraphUtils.h"
+
+namespace UE::NNERuntimeRDG::Private::Hlsl
+{
+	DECLARE_GPU_STAT_NAMED(FNNEOperatorReshape, TEXT("NNE.Operator.Hlsl.Reshape"));
+
+	/**
+	 * Reshape operator implementation
+	 */
+	class FReshape : public FOperatorHlsl
+	{
+	public:
+
+		FReshape() {}
+		virtual ~FReshape() = default;
+
+	public:
+
+		virtual int PrepareOutputs(TConstArrayView<NNECore::Internal::FTensorRef> InputTensors, TArrayView<NNECore::Internal::FTensorRef> OutputTensors) const override
+		{
+			check(InputTensors.Num() == 2);
+			check(OutputTensors.Num() == 1);
+
+			const NNECore::Internal::FTensor& X = *InputTensors[0];
+			const NNECore::Internal::FTensor& ShapeTensor = *InputTensors[1];
+			TArray<uint32> OutputShapeData;
+
+			check(ShapeTensor.GetDataType() == ENNETensorDataType::Int64);
+
+			if (!ShapeTensor.HasPreparedData())
+			{
+				UE_LOG(LogNNE, Warning, TEXT("Reshape input 'Shape' (name: %s) should be constant for shape inference to succeed, however it is not."), *ShapeTensor.GetName());
+				return -1;
+			}
+			for (int64 Dim : ShapeTensor.GetPreparedData<int64>())
+			{
+				if (Dim <= 0)
+				{
+					UE_LOG(LogNNE, Warning, TEXT("Cannot reshape 'Data' (name: %s) to a shape with a dimension less or equal to 0."), *X.GetName());
+					return -1;
+				}
+				OutputShapeData.Add((uint32)Dim);
+			}
+
+			NNECore::FTensorShape OutputShape = NNECore::FTensorShape::Make(OutputShapeData);
+			
+			OutputTensors[0]->SetShape(OutputShape);
+			if (X.HasPreparedData())
+			{
+				OutputTensors[0]->SetPreparedData<uint8>(X.GetPreparedData<uint8>());
+			}
+			
+			return 0;
+		};
+
+		virtual bool Initialize(TConstArrayView<NNECore::FTensorDesc> InputTensorDescs, TConstArrayView<NNECore::FTensorDesc> OutputTensorDescs, const NNECore::FAttributeMap& Attributes) override
+		{
+			check(InputTensorDescs.Num() == 2);
+			check(OutputTensorDescs.Num() == 1);
+			
+			return true;
+		}
+
+		virtual void Dispatch(FRDGBuilder& GraphBuilder, TConstArrayView<FTensorRDGRef> InputTensors, TConstArrayView<FTensorRDGRef> OutputTensors) override
+		{
+			check(InputTensors.Num() == 2);
+			check(OutputTensors.Num() == 1);
+			check(InputTensors[0] != nullptr);
+			check(OutputTensors[0] != nullptr);
+
+			const FTensorRDG& Data = *InputTensors[0];
+			const FTensorRDG& Output = *OutputTensors[0];
+
+			RDG_EVENT_SCOPE(GraphBuilder, "NNE.Operator.Hlsl.Reshape");
+			RDG_GPU_STAT_SCOPE(GraphBuilder, FNNEOperatorReshape);
+
+			AddCopyBufferPass(GraphBuilder, Output.GetBuffer(), Data.GetBuffer());
+		}
+	};
+
+	bool ValidateReshapeOperator(const NNECore::FAttributeMap& AttributeMap, TConstArrayView<ENNETensorDataType> InputTypes, TConstArrayView<NNECore::FSymbolicTensorShape> InputReshapes)
+	{
+		bool bIsValid = true;
+
+		//This match version 5 of the Reshape operator, next version is 13
+		//[REDACTED_PROFILE_URL]/onnx/blob/main/docs/Operators.md#Reshape
+		FAttributeValidator AttributeValidator;
+		bIsValid &= AttributeValidator.Validate(AttributeMap);
+
+		FInputValidator InputValidator;
+		InputValidator.SetTemplateCount(2);
+		InputValidator.AddSupportedType(ENNETensorDataType::Half);
+		InputValidator.AddSupportedType(ENNETensorDataType::Float);
+		InputValidator.AddSupportedType(ENNETensorDataType::Double);
+		InputValidator.AddSupportedType(ENNETensorDataType::Int8);
+		InputValidator.AddSupportedType(ENNETensorDataType::Int16);
+		InputValidator.AddSupportedType(ENNETensorDataType::Int32);
+		InputValidator.AddSupportedType(ENNETensorDataType::Int64);
+		InputValidator.AddSupportedType(ENNETensorDataType::UInt8);
+		InputValidator.AddSupportedType(ENNETensorDataType::UInt16);
+		InputValidator.AddSupportedType(ENNETensorDataType::UInt32);
+		InputValidator.AddSupportedType(ENNETensorDataType::UInt64);
+		InputValidator.AddRequired();
+
+		InputValidator.AddSupportedType(ENNETensorDataType::Int64, 1);
+		InputValidator.AddRequired(1);
+
+		bIsValid &= InputValidator.Validate(InputTypes);
+
+		return bIsValid;
+	}
+
+	FOperatorHlsl* CreateReshapeOperator()
+	{
+		return new FReshape();
+	}
+
+	bool RegisterReshapeOperator(FOperatorRegistryHlsl& Registry)
+	{
+		Registry.OpAdd(TEXT("Reshape"), CreateReshapeOperator, ValidateReshapeOperator);
+		return true;
+	}
+} // UE::NNERuntimeRDG::Private::Hlsl
diff --git a/Engine/Plugins/Experimental/NNE/Source/NNERuntimeRDG/Private/Hlsl/NNERuntimeRDGReshape.h b/Engine/Plugins/Experimental/NNE/Source/NNERuntimeRDG/Private/Hlsl/NNERuntimeRDGReshape.h
new file mode 100644
--- /dev/null
+++ b/Engine/Plugins/Experimental/NNE/Source/NNERuntimeRDG/Private/Hlsl/NNERuntimeRDGReshape.h
@@ -0,0 +1,11 @@
+// Copyright Epic Games, Inc. All Rights Reserved.
+
+#pragma once
+
+#include "NNERuntimeRDG.h"
+#include "NNERuntimeRDGHlslOp.h"
+
+namespace UE::NNERuntimeRDG::Private::Hlsl
+{
+	bool RegisterReshapeOperator(FOperatorRegistryHlsl& Registry);
+} // UE::NNERuntimeRDG::Private::Hlsl
diff --git a/Engine/Plugins/Experimental/NNE/Source/NNERuntimeRDG/Private/Hlsl/NNERuntimeRDGSize.cpp b/Engine/Plugins/Experimental/NNE/Source/NNERuntimeRDG/Private/Hlsl/NNERuntimeRDGSize.cpp
new file mode 100644
--- /dev/null
+++ b/Engine/Plugins/Experimental/NNE/Source/NNERuntimeRDG/Private/Hlsl/NNERuntimeRDGSize.cpp
@@ -0,0 +1,92 @@
+// Copyright Epic Games, Inc. All Rights Reserved.
+
+#include "NNERuntimeRDGSize.h"
+#include "NNECoreTensor.h"
+#include "NNECoreTypes.h"
+
+namespace UE::NNERuntimeRDG::Private::Hlsl
+{
+	/**
+	 * Size operator implementation
+	 */
+	class FSize : public FOperatorHlsl
+	{
+	public:
+
+		FSize() {}
+		virtual ~FSize() = default;
+
+		TArray<int32> Axes;
+
+	public:
+
+		virtual int PrepareOutputs(TConstArrayView<NNECore::Internal::FTensorRef> InputTensors, TArrayView<NNECore::Internal::FTensorRef> OutputTensors) const override
+		{
+			check(InputTensors.Num() == 1);
+			check(OutputTensors.Num() == 1);
+
+			const NNECore::Internal::FTensor& X = *InputTensors[0];
+			NNECore::FTensorShape ScalarShape = NNECore::FTensorShape::Make({});
+			TArray<int64> OutputData;
+			
+			OutputTensors[0]->SetShape(ScalarShape);
+			OutputData.Add(X.GetShape().Volume());
+			OutputTensors[0]->SetPreparedData<int64>(OutputData);
+			
+			return 0;
+		};
+
+		virtual bool Initialize(TConstArrayView<NNECore::FTensorDesc> InputTensorDescs, TConstArrayView<NNECore::FTensorDesc> OutputTensorDescs, const NNECore::FAttributeMap& Attributes) override
+		{
+			check(InputTensorDescs.Num() == 1);
+			check(OutputTensorDescs.Num() == 1);
+			check(OutputTensorDescs[0].GetDataType() == ENNETensorDataType::Int64);
+
+			return true;
+		}
+
+		virtual void Dispatch(FRDGBuilder& GraphBuilder, TConstArrayView<FTensorRDGRef> InputTensors, TConstArrayView<FTensorRDGRef> OutputTensors) override
+		{
+			UE_LOG(LogNNE, Warning, TEXT("Size: Output should be constant and already uploaded to GPU memory. Dispatch should not need to be called."));
+		}
+	};
+
+	bool ValidateSizeOperator(const NNECore::FAttributeMap& AttributeMap, TConstArrayView<ENNETensorDataType> InputTypes, TConstArrayView<NNECore::FSymbolicTensorShape> InputSizes)
+	{
+		bool bIsValid = true;
+
+		//This match version 13 of the Size operator, next version are 11 and 13
+		//[REDACTED_PROFILE_URL]/onnx/blob/main/docs/Operators.md#Size
+		FAttributeValidator AttributeValidator;
+		bIsValid &= AttributeValidator.Validate(AttributeMap);
+
+		FInputValidator InputValidator;
+		InputValidator.AddSupportedType(ENNETensorDataType::Half);
+		InputValidator.AddSupportedType(ENNETensorDataType::Float);
+		InputValidator.AddSupportedType(ENNETensorDataType::Double);
+		InputValidator.AddSupportedType(ENNETensorDataType::Int8);
+		InputValidator.AddSupportedType(ENNETensorDataType::Int16);
+		InputValidator.AddSupportedType(ENNETensorDataType::Int32);
+		InputValidator.AddSupportedType(ENNETensorDataType::Int64);
+		InputValidator.AddSupportedType(ENNETensorDataType::UInt8);
+		InputValidator.AddSupportedType(ENNETensorDataType::UInt16);
+		InputValidator.AddSupportedType(ENNETensorDataType::UInt32);
+		InputValidator.AddSupportedType(ENNETensorDataType::UInt64);
+		InputValidator.AddRequired();
+
+		bIsValid &= InputValidator.Validate(InputTypes);
+
+		return bIsValid;
+	}
+
+	FOperatorHlsl* CreateSizeOperator()
+	{
+		return new FSize();
+	}
+
+	bool RegisterSizeOperator(FOperatorRegistryHlsl& Registry)
+	{
+		Registry.OpAdd(TEXT("Size"), CreateSizeOperator, ValidateSizeOperator);
+		return true;
+	}
+} // UE::NNERuntimeRDG::Private::Hlsl
diff --git a/Engine/Plugins/Experimental/NNE/Source/NNERuntimeRDG/Private/Hlsl/NNERuntimeRDGSize.h b/Engine/Plugins/Experimental/NNE/Source/NNERuntimeRDG/Private/Hlsl/NNERuntimeRDGSize.h
new file mode 100644
--- /dev/null
+++ b/Engine/Plugins/Experimental/NNE/Source/NNERuntimeRDG/Private/Hlsl/NNERuntimeRDGSize.h
@@ -0,0 +1,11 @@
+// Copyright Epic Games, Inc. All Rights Reserved.
+
+#pragma once
+
+#include "NNERuntimeRDG.h"
+#include "NNERuntimeRDGHlslOp.h"
+
+namespace UE::NNERuntimeRDG::Private::Hlsl
+{
+	bool RegisterSizeOperator(FOperatorRegistryHlsl& Registry);
+} // UE::NNERuntimeRDG::Private::Hlsl
diff --git a/Engine/Plugins/Experimental/NNE/Source/NNERuntimeRDG/Private/Hlsl/NNERuntimeRDGSqueeze.cpp b/Engine/Plugins/Experimental/NNE/Source/NNERuntimeRDG/Private/Hlsl/NNERuntimeRDGSqueeze.cpp
new file mode 100644
--- /dev/null
+++ b/Engine/Plugins/Experimental/NNE/Source/NNERuntimeRDG/Private/Hlsl/NNERuntimeRDGSqueeze.cpp
@@ -0,0 +1,137 @@
+// Copyright Epic Games, Inc. All Rights Reserved.
+
+#include "NNERuntimeRDGSqueeze.h"
+#include "NNECoreTensor.h"
+#include "NNECoreTypes.h"
+#include "RenderGraphUtils.h"
+
+namespace UE::NNERuntimeRDG::Private::Hlsl
+{
+	DECLARE_GPU_STAT_NAMED(FNNEOperatorSqueeze, TEXT("NNE.Operator.Hlsl.Squeeze"));
+
+	/**
+	 * Squeeze operator implementation
+	 */
+	class FSqueeze : public FOperatorHlsl
+	{
+	public:
+
+		FSqueeze() {}
+		virtual ~FSqueeze() = default;
+
+		TArray<int32> Axes;
+
+	public:
+
+		virtual int PrepareOutputs(TConstArrayView<NNECore::Internal::FTensorRef> InputTensors, TArrayView<NNECore::Internal::FTensorRef> OutputTensors) const override
+		{
+			check(InputTensors.Num() == 1);
+			check(OutputTensors.Num() == 1);
+
+			const NNECore::Internal::FTensor& X = *InputTensors[0];
+			TArray<uint32> OutputShapeData(X.GetShape().GetData());
+
+			for (int32 Axe : Axes)
+			{
+				if (OutputShapeData[Axe] != 1)
+				{
+					UE_LOG(LogNNE, Warning, TEXT("Squeeze at axe %d for 'Data' (name: %s) should be targeting a dimension of size 1 but it is %d."), Axe, *X.GetName(), OutputShapeData[Axe]);
+					return -1;
+				}
+
+				OutputShapeData.RemoveAt(Axe);
+			}
+			
+			NNECore::FTensorShape OutputShape = NNECore::FTensorShape::Make(OutputShapeData);
+			
+			OutputTensors[0]->SetShape(OutputShape);
+			if (X.HasPreparedData())
+			{
+				OutputTensors[0]->SetPreparedData<uint8>(X.GetPreparedData<uint8>());
+			}
+			
+			return 0;
+		};
+
+		virtual bool Initialize(TConstArrayView<NNECore::FTensorDesc> InputTensorDescs, TConstArrayView<NNECore::FTensorDesc> OutputTensorDescs, const NNECore::FAttributeMap& Attributes) override
+		{
+			check(InputTensorDescs.Num() == 1);
+			check(OutputTensorDescs.Num() == 1);
+
+			Axes = Attributes.GetValue<TArray<int32>>(TEXT("axes"));
+			for (int32 Axe : Axes)
+			{
+				if (Axe < 0)
+				{
+					UE_LOG(LogNNE, Warning, TEXT("Squeeze operator does not support negative axes"));
+					return false;
+				}
+				if (Axe >= InputTensorDescs[0].GetShape().Rank())
+				{
+					UE_LOG(LogNNE, Warning, TEXT("Squeeze operator does not support axes greater than the number of dimensions of the input tensor shape"));
+					return false;
+				}
+			}
+
+			Axes.Sort();
+			
+			return true;
+		}
+
+		virtual void Dispatch(FRDGBuilder& GraphBuilder, TConstArrayView<FTensorRDGRef> InputTensors, TConstArrayView<FTensorRDGRef> OutputTensors) override
+		{
+			check(InputTensors.Num() == 1);
+			check(OutputTensors.Num() == 1);
+			check(InputTensors[0] != nullptr);
+			check(OutputTensors[0] != nullptr);
+			
+			const FTensorRDG& Data = *InputTensors[0];
+			const FTensorRDG& Output = *OutputTensors[0];
+
+			RDG_EVENT_SCOPE(GraphBuilder, "NNE.Operator.Hlsl.Squeeze");
+			RDG_GPU_STAT_SCOPE(GraphBuilder, FNNEOperatorSqueeze);
+
+			AddCopyBufferPass(GraphBuilder, Output.GetBuffer(), Data.GetBuffer());
+		}
+	};
+
+	bool ValidateSqueezeOperator(const NNECore::FAttributeMap& AttributeMap, TConstArrayView<ENNETensorDataType> InputTypes, TConstArrayView<NNECore::FSymbolicTensorShape> InputSqueezes)
+	{
+		bool bIsValid = true;
+
+		//This match version 1 of the Squeeze operator, next version are 11 and 13
+		//[REDACTED_PROFILE_URL]/onnx/blob/main/docs/Operators.md#Squeeze
+		FAttributeValidator AttributeValidator;
+		AttributeValidator.AddRequired(TEXT("axes"), ENNEAttributeDataType::Int32Array);
+		bIsValid &= AttributeValidator.Validate(AttributeMap);
+
+		FInputValidator InputValidator;
+		InputValidator.AddSupportedType(ENNETensorDataType::Half);
+		InputValidator.AddSupportedType(ENNETensorDataType::Float);
+		InputValidator.AddSupportedType(ENNETensorDataType::Double);
+		InputValidator.AddSupportedType(ENNETensorDataType::Int8);
+		InputValidator.AddSupportedType(ENNETensorDataType::Int16);
+		InputValidator.AddSupportedType(ENNETensorDataType::Int32);
+		InputValidator.AddSupportedType(ENNETensorDataType::Int64);
+		InputValidator.AddSupportedType(ENNETensorDataType::UInt8);
+		InputValidator.AddSupportedType(ENNETensorDataType::UInt16);
+		InputValidator.AddSupportedType(ENNETensorDataType::UInt32);
+		InputValidator.AddSupportedType(ENNETensorDataType::UInt64);
+		InputValidator.AddRequired();
+
+		bIsValid &= InputValidator.Validate(InputTypes);
+
+		return bIsValid;
+	}
+
+	FOperatorHlsl* CreateSqueezeOperator()
+	{
+		return new FSqueeze();
+	}
+
+	bool RegisterSqueezeOperator(FOperatorRegistryHlsl& Registry)
+	{
+		Registry.OpAdd(TEXT("Squeeze"), CreateSqueezeOperator, ValidateSqueezeOperator);
+		return true;
+	}
+} // UE::NNERuntimeRDG::Private::Hlsl
diff --git a/Engine/Plugins/Experimental/NNE/Source/NNERuntimeRDG/Private/Hlsl/NNERuntimeRDGSqueeze.h b/Engine/Plugins/Experimental/NNE/Source/NNERuntimeRDG/Private/Hlsl/NNERuntimeRDGSqueeze.h
new file mode 100644
--- /dev/null
+++ b/Engine/Plugins/Experimental/NNE/Source/NNERuntimeRDG/Private/Hlsl/NNERuntimeRDGSqueeze.h
@@ -0,0 +1,11 @@
+// Copyright Epic Games, Inc. All Rights Reserved.
+
+#pragma once
+
+#include "NNERuntimeRDG.h"
+#include "NNERuntimeRDGHlslOp.h"
+
+namespace UE::NNERuntimeRDG::Private::Hlsl
+{
+	bool RegisterSqueezeOperator(FOperatorRegistryHlsl& Registry);
+} // UE::NNERuntimeRDG::Private::Hlsl
diff --git a/Engine/Plugins/Experimental/NNE/Source/NNERuntimeRDG/Private/Hlsl/NNERuntimeRDGUnsqueeze.cpp b/Engine/Plugins/Experimental/NNE/Source/NNERuntimeRDG/Private/Hlsl/NNERuntimeRDGUnsqueeze.cpp
--- a/Engine/Plugins/Experimental/NNE/Source/NNERuntimeRDG/Private/Hlsl/NNERuntimeRDGUnsqueeze.cpp
+++ b/Engine/Plugins/Experimental/NNE/Source/NNERuntimeRDG/Private/Hlsl/NNERuntimeRDGUnsqueeze.cpp
@@ -3,9 +3,12 @@
 #include "NNERuntimeRDGUnsqueeze.h"
 #include "NNECoreTensor.h"
 #include "NNECoreTypes.h"
+#include "RenderGraphUtils.h"
 
 namespace UE::NNERuntimeRDG::Private::Hlsl
 {
+	DECLARE_GPU_STAT_NAMED(FNNEOperatorUnsqueeze, TEXT("NNE.Operator.Hlsl.Unsqueeze"));
+
 	/**
 	 * Unsqueeze operator implementation
 	 */
@@ -28,12 +31,6 @@ namespace UE::NNERuntimeRDG::Private::Hlsl
 			const NNECore::Internal::FTensor& X = *InputTensors[0];
 			TArray<uint32> OutputShapeData(X.GetShape().GetData());
 
-			if (!X.HasPreparedData())
-			{
-				UE_LOG(LogNNE, Warning, TEXT("Unsqueeze input 'Data' (name: %s) should be constant, Unsqueeze is not implemented on GPU at the moment."), *X.GetName());
-				return -1;
-			}
-			
 			for (int32 Axe : Axes)
 			{
 				OutputShapeData.Insert(1, Axe);
@@ -42,7 +39,10 @@ namespace UE::NNERuntimeRDG::Private::Hlsl
 			NNECore::FTensorShape OutputShape = NNECore::FTensorShape::Make(OutputShapeData);
 			
 			OutputTensors[0]->SetShape(OutputShape);
-			OutputTensors[0]->SetPreparedData<uint8>(X.GetPreparedData<uint8>());
+			if (X.HasPreparedData())
+			{
+				OutputTensors[0]->SetPreparedData<uint8>(X.GetPreparedData<uint8>());
+			}
 			
 			return 0;
 		};
@@ -66,13 +66,26 @@ namespace UE::NNERuntimeRDG::Private::Hlsl
 					return false;
 				}
 			}
+
+			Axes.Sort();
 			
 			return true;
 		}
 
 		virtual void Dispatch(FRDGBuilder& GraphBuilder, TConstArrayView<FTensorRDGRef> InputTensors, TConstArrayView<FTensorRDGRef> OutputTensors) override
 		{
-			UE_LOG(LogNNE, Warning, TEXT("Unsqueeze: Output should be constant and already uploaded to GPU memory. Dispatch should not need to be called."));
+			check(InputTensors.Num() == 1);
+			check(OutputTensors.Num() == 1);
+			check(InputTensors[0] != nullptr);
+			check(OutputTensors[0] != nullptr);
+
+			const FTensorRDG& Data = *InputTensors[0];
+			const FTensorRDG& Output = *OutputTensors[0];
+
+			RDG_EVENT_SCOPE(GraphBuilder, "NNE.Operator.Hlsl.Unsqueeze");
+			RDG_GPU_STAT_SCOPE(GraphBuilder, FNNEOperatorUnsqueeze);
+
+			AddCopyBufferPass(GraphBuilder, Output.GetBuffer(), Data.GetBuffer());
 		}
 	};
 
diff --git a/Engine/Plugins/Experimental/NNE/Source/NNERuntimeRDG/Private/NNERuntimeRDGHlsl.cpp b/Engine/Plugins/Experimental/NNE/Source/NNERuntimeRDG/Private/NNERuntimeRDGHlsl.cpp
--- a/Engine/Plugins/Experimental/NNE/Source/NNERuntimeRDG/Private/NNERuntimeRDGHlsl.cpp
+++ b/Engine/Plugins/Experimental/NNE/Source/NNERuntimeRDG/Private/NNERuntimeRDGHlsl.cpp
@@ -19,8 +19,11 @@
 #include "Hlsl/NNERuntimeRDGGemm.h"
 #include "Hlsl/NNERuntimeRDGInstanceNormalization.h"
 #include "Hlsl/NNERuntimeRDGPad.h"
+#include "Hlsl/NNERuntimeRDGReshape.h"
 #include "Hlsl/NNERuntimeRDGShape.h"
+#include "Hlsl/NNERuntimeRDGSize.h"
 #include "Hlsl/NNERuntimeRDGSlice.h"
+#include "Hlsl/NNERuntimeRDGSqueeze.h"
 #include "Hlsl/NNERuntimeRDGUnsqueeze.h"
 #include "Hlsl/NNERuntimeRDGUpsample.h"
 #include "Hlsl/NNERuntimeRDGMatMul.h"
@@ -35,22 +38,25 @@ bool UNNERuntimeRDGHlslImpl::Init()
 	FOperatorRegistryHlsl* registry = FOperatorRegistryHlsl::Get();
 	check(registry != nullptr);
 
-	RegisterElementWiseUnaryOperators(*registry);
+	RegisterCastOperator(*registry);
+	RegisterConvOperator(*registry);
+	RegisterConcatOperator(*registry);
+	RegisterConvTransposeOperator(*registry);
 	RegisterElementWiseBinaryOperators(*registry);
+	RegisterElementWiseUnaryOperators(*registry);
 	RegisterElementWiseVariadicOperators(*registry);
+	RegisterGatherOperator(*registry);
 	RegisterGemmOperator(*registry);
-	RegisterConvOperator(*registry);
-	RegisterConvTransposeOperator(*registry);
-	RegisterMatMulOperator(*registry);
 	RegisterInstanceNormalizationOperator(*registry);
-	RegisterUpsampleOperator(*registry);
 	RegisterPadOperator(*registry);
+	RegisterReshapeOperator(*registry);
 	RegisterShapeOperator(*registry);
+	RegisterSizeOperator(*registry);
 	RegisterSliceOperator(*registry);
-	RegisterConcatOperator(*registry);
+	RegisterSqueezeOperator(*registry);
 	RegisterUnsqueezeOperator(*registry);
-	RegisterCastOperator(*registry);
-	RegisterGatherOperator(*registry);
+	RegisterUpsampleOperator(*registry);
+	RegisterMatMulOperator(*registry);
 
 	return true;
 }
diff --git a/Engine/Plugins/Experimental/NNE/Source/NNEUtils/Private/NNEUtilsModelBuilderONNX.cpp b/Engine/Plugins/Experimental/NNE/Source/NNEUtils/Private/NNEUtilsModelBuilderONNX.cpp
--- a/Engine/Plugins/Experimental/NNE/Source/NNEUtils/Private/NNEUtilsModelBuilderONNX.cpp
+++ b/Engine/Plugins/Experimental/NNE/Source/NNEUtils/Private/NNEUtilsModelBuilderONNX.cpp
@@ -354,6 +354,8 @@ NNEUTILS_API bool CreateONNXModelForOperator(bool UseVariadicShapeForModel, cons
 		OperatorName == TEXT("Pad") ||		// current implementation is opset 2 (next version is 11)
 		OperatorName == TEXT("Slice") ||	// current implementation is opset 1 (next version is 10)
 		OperatorName == TEXT("Shape") ||	// current implementation is opset 1 (next version is 13)
+		OperatorName == TEXT("Reshape") ||	// current implementation is opset 5 (next version is 13)
+		OperatorName == TEXT("Squeeze") ||	// current implementation is opset 1 (next version is 11)
 		OperatorName == TEXT("Unsqueeze"))	// current implementation is opset 1 (next version is 11)
 	{
 		OpsetVersion = 9;
