NNE Cleanup: Cast op have it's implementation in a dedicated .h/.cpp.

#rnx

diff --git a/Engine/Plugins/Experimental/NNE/Source/NNECore/Internal/NNECoreOperator.h b/Engine/Plugins/Experimental/NNE/Source/NNECore/Internal/NNECoreOperator.h
--- a/Engine/Plugins/Experimental/NNE/Source/NNECore/Internal/NNECoreOperator.h
+++ b/Engine/Plugins/Experimental/NNE/Source/NNECore/Internal/NNECoreOperator.h
@@ -18,7 +18,6 @@ namespace UE::NNECore::Internal
 		Atan,
 		Atanh,
 		//BitShift, //Note: need integer tensors
-		//Cast,     //Note: ability to cast tensor type
 		Ceil,
 		//Clip,     //Note: need scalar tensor inputs
 		Cos,
diff --git a/Engine/Plugins/Experimental/NNE/Source/NNERuntimeRDG/Private/Hlsl/NNERuntimeRDGElementWiseUnary.cpp b/Engine/Plugins/Experimental/NNE/Source/NNERuntimeRDG/Private/Hlsl/NNERuntimeRDGElementWiseUnary.cpp
--- a/Engine/Plugins/Experimental/NNE/Source/NNERuntimeRDG/Private/Hlsl/NNERuntimeRDGElementWiseUnary.cpp
+++ b/Engine/Plugins/Experimental/NNE/Source/NNERuntimeRDG/Private/Hlsl/NNERuntimeRDGElementWiseUnary.cpp
@@ -220,7 +220,6 @@ namespace UE::NNERuntimeRDG::Private::Hlsl
 		OP(Atan);
 		OP(Atanh);
 		//OP(BitShift);
-		//OP(Cast);
 		OP(Ceil);
 		//OP(Clip);
 		OP(Cos);
