NNE - Fix UpdateShapeAndStrides()'s impl which is breaking Transpose.
#rnx

diff --git a/Engine/Plugins/Experimental/NNE/Source/NNERuntimeRDG/Private/Dml/NNEDmlOperator.cpp b/Engine/Plugins/Experimental/NNE/Source/NNERuntimeRDG/Private/Dml/NNEDmlOperator.cpp
--- a/Engine/Plugins/Experimental/NNE/Source/NNERuntimeRDG/Private/Dml/NNEDmlOperator.cpp
+++ b/Engine/Plugins/Experimental/NNE/Source/NNERuntimeRDG/Private/Dml/NNEDmlOperator.cpp
@@ -69,7 +69,7 @@ namespace DmlUtil
 
 	void FTensorDesc::UpdateShapeAndStrides(TConstArrayView<uint32> InShape, TConstArrayView<uint32> InStrides)
 	{
-		check(!InStrides.IsEmpty() && (InStrides.Num() == InShape.Num()));
+		check(!InShape.IsEmpty());
 
 		Sizes = InShape;
 		BuffDesc.Sizes = Sizes.GetData();
@@ -77,8 +77,9 @@ namespace DmlUtil
 
 		if (!InStrides.IsEmpty())
 		{
+			check(InStrides.Num() == InShape.Num());
 			Strides = InStrides;
-			BuffDesc.Strides = nullptr;
+			BuffDesc.Strides = Strides.GetData();
 		}
 		else
 		{
