NNE - Fix DML Slice data type
#rnx

diff --git a/Engine/Plugins/Experimental/NNE/Source/NNERuntimeRDG/Private/Dml/NNEDmlOperatorSlice.cpp b/Engine/Plugins/Experimental/NNE/Source/NNERuntimeRDG/Private/Dml/NNEDmlOperatorSlice.cpp
--- a/Engine/Plugins/Experimental/NNE/Source/NNERuntimeRDG/Private/Dml/NNEDmlOperatorSlice.cpp
+++ b/Engine/Plugins/Experimental/NNE/Source/NNERuntimeRDG/Private/Dml/NNEDmlOperatorSlice.cpp
@@ -4,7 +4,7 @@
 
 #include "NNEDmlOperator.h"
 #include "Misc/EnumerateRange.h"
-#include "Algo/Transform.h"
+#include "Algo/Copy.h"
 #include "Algo/ForEach.h"
 
 #include <numeric>
@@ -16,66 +16,31 @@ namespace UE::NNERuntimeRDG::Private::Dml
  * Slice
  */
 class FOperatorDmlSlice : public FOperatorDml
-{
-	
-
-public:
-
-	static FOperatorDml* Create()
+{	
+	template<typename DataType>
+	void ComputeOffsetsSizesStrides(
+		TArrayView<const NNECore::Internal::FTensor> InputTensors,
+		DmlUtil::FSmallUIntArray& OutOffsets, 
+		DmlUtil::FSmallUIntArray& OutSizes, 
+		DmlUtil::FSmallIntArray& OutStrides)
 	{
-		return new FOperatorDmlSlice();
-	}
-
-	//
-	//
-	//
-	virtual bool Initialize(IDMLDevice* Device, TArrayView<const NNECore::Internal::FTensor> InputTensors, TArrayView<const NNECore::Internal::FTensor> OutputTensors, const NNECore::FAttributeMap& Attributes) override
-	{
-		check(InputTensors.Num() >= 3);
-		check(InputTensors.Num() <= 5);
-		check(OutputTensors.Num() == 1);
-
-		for(int Idx = 1; Idx < InputTensors.Num(); Idx++)
-		{
-			check(InputTensors[Idx].GetShape().Rank() == 1);
-			check(InputTensors[Idx].GetDataType() == ENNETensorDataType::Int32 
-			   || InputTensors[Idx].GetDataType() == ENNETensorDataType::Int64);
-			check(InputTensors[Idx].HasPreparedData());
-			ConstantCPUInputs.Add(Idx);
-		}
-
-		// Initialize Input tensor desc
-        DmlUtil::FTensorDesc DmlInputTensorDesc;
-        if (!DmlInputTensorDesc.InitFromTensor(InputTensors[0], InputTensors[0].GetShape().Rank()))
-        {
-            UE_LOG(LogNNE, Error, TEXT("Failed to initialize Slice input for DML inference"));
-            return false;
-        }
-
-		// Initialize Output tensor desc
-        DmlUtil::FTensorDesc DmlOutputTensorDesc;
-        if (!DmlOutputTensorDesc.InitFromTensor(OutputTensors[0], OutputTensors[0].GetShape().Rank()))
-        {
-            UE_LOG(LogNNE, Error, TEXT("Failed to initialize Slice Output for DML inference"));
-            return false;
-        }
-		TConstArrayView<int32> Starts = InputTensors[1].GetPreparedData<int32>();
-		TConstArrayView<int32> Ends = InputTensors[2].GetPreparedData<int32>();
+		TConstArrayView<DataType> Starts =  InputTensors[1].GetPreparedData<DataType>();
+		TConstArrayView<DataType> Ends = InputTensors[2].GetPreparedData<DataType>();
 		check(Starts.Num() == Ends.Num());
 
-		DmlUtil::FSmallIntArray Axes;
-		DmlUtil::FSmallIntArray Steps;
+		DmlUtil::FSmallArray<DataType> Axes;
+		DmlUtil::FSmallArray<DataType> Steps;
 
 		if(InputTensors.Num() >= 4)
 		{
-			auto NormalizeAxes = [NumDims = InputTensors[0].GetShape().Rank()] (int32& Axis)
+			auto NormalizeAxes = [NumDims = (DataType) InputTensors[0].GetShape().Rank()] (DataType& Axis)
 			{
 				if(Axis < 0)
 				{
 					Axis += NumDims;
 				}
 			};
-			Axes = InputTensors[3].GetPreparedData<int32>();
+			Axes = InputTensors[3].GetPreparedData<DataType>();
 			check(Axes.Num() == Starts.Num());
 			Algo::ForEach(Axes, NormalizeAxes);
 		}
@@ -87,59 +52,124 @@ public:
 
 		if(InputTensors.Num() >= 5)
 		{
-			Steps = InputTensors[4].GetPreparedData<int32>();
+			Steps = InputTensors[4].GetPreparedData<DataType>();
 			check(Steps.Num() == Axes.Num());
 		}
 		else
 		{
-			Steps.Init(1, Axes.Num());
+			Steps.Init((DataType)1, Axes.Num());
 		}
 
-		DmlUtil::FSmallUIntArray OutputShape, Sizes, Offsets;
-		Algo::Transform(InputTensors[0].GetShape().GetData(), OutputShape, [](int32 In) { return (uint32) In; });
-		Sizes = OutputShape;
-		Offsets.SetNumZeroed(OutputShape.Num());
+		DmlUtil::FSmallUIntArray OutputShape;
+		Algo::Copy(InputTensors[0].GetShape().GetData(), OutputShape);
+		OutSizes = OutputShape;
+		OutOffsets.SetNumZeroed(OutputShape.Num());
 
-		DmlUtil::FSmallIntArray Strides;
-		Strides.Init(1, OutputShape.Num());
+		OutStrides.Init(1, OutputShape.Num());
 
-		for (TConstEnumerateRef<int32> Elem : EnumerateRange(Starts))
+		for (TConstEnumerateRef<DataType> Elem : EnumerateRange(Starts))
 		{
 			int32 Idx = Elem.GetIndex();
-			int32 Start = *Elem;
-			int32 End = Ends[Idx];
+			DataType Start = *Elem;
+			DataType End = Ends[Idx];
 			
-			int32 DimIndex = Axes[Idx];
-			check(DimIndex < InputTensors[0].GetShape().Rank());
-			int32 Stride = Steps[Idx];
+			DataType DimIndex = Axes[Idx];
+			check(DimIndex < (DataType) InputTensors[0].GetShape().Rank());
+			DataType Stride = Steps[Idx];
 			check(Stride != 0);
 
-			int32 Dim = InputTensors[0].GetShape().GetData()[DimIndex];
-			if(Start < 0 && Start > TNumericLimits<int32>::Min())
+			uint32 Dim = InputTensors[0].GetShape().GetData()[DimIndex];
+			if(Start < 0 && Start > TNumericLimits<DataType>::Min())
 			{
-				Start += Dim;
+				Start += (DataType) Dim;
 			}
-			if(End < 0 && Start > TNumericLimits<int32>::Min())
+			if(End < 0 && Start > TNumericLimits<DataType>::Min())
 			{
-				End += Dim;
+				End += (DataType) Dim;
 			}
 
 			if (Stride < 0)
             {
                 std::swap(Start, End);
-                Start += (Start < TNumericLimits<int32>::Max()) ? 1 : 0;
-                End += (End < TNumericLimits<int32>::Max()) ? 1 : 0;
+                Start += (Start < TNumericLimits<DataType>::Max()) ? 1 : 0;
+                End += (End < TNumericLimits<DataType>::Max()) ? 1 : 0;
             }
 
 			Start = FMath::Max(Start, 0);
-            End = FMath::Min(End, Dim);
-            int32 Size = FMath::Max(End - Start, 0);
+            End = FMath::Min(End, (DataType) Dim);
+            DataType Size = FMath::Max(End - Start, 0);
 			
-			int32 AbsStride = FMath::Abs(Stride);
+			DataType AbsStride = FMath::Abs(Stride);
             OutputShape[DimIndex] = (uint32) ((Size / AbsStride) + (Size % AbsStride != 0));
-            Offsets[DimIndex] = (uint32) Start;
-            Strides[DimIndex] = Stride;
-            Sizes[DimIndex] = (uint32) Size;
+            OutOffsets[DimIndex] = (uint32) Start;
+            OutStrides[DimIndex] = (int32) Stride;
+            OutSizes[DimIndex] = (uint32) Size;
+		}
+	}
+
+public:
+
+	static FOperatorDml* Create()
+	{
+		return new FOperatorDmlSlice();
+	}
+
+	//
+	//
+	//
+	virtual bool Initialize(IDMLDevice* Device, TArrayView<const NNECore::Internal::FTensor> InputTensors, TArrayView<const NNECore::Internal::FTensor> OutputTensors, const NNECore::FAttributeMap& Attributes) override
+	{
+		check(InputTensors.Num() >= 3);
+		check(InputTensors.Num() <= 5);
+		check(OutputTensors.Num() == 1);
+
+		ENNETensorDataType InputIndexDataType = ENNETensorDataType::None;
+
+		for(int Idx = 1; Idx < InputTensors.Num(); Idx++)
+		{
+			check(InputTensors[Idx].GetShape().Rank() == 1);
+			if(Idx == 1)
+			{
+				InputIndexDataType = InputTensors[Idx].GetDataType();
+				check( InputIndexDataType == ENNETensorDataType::Int32 
+			   		|| InputIndexDataType == ENNETensorDataType::Int64);
+			}
+			else
+			{
+				check(InputTensors[Idx].GetDataType() == InputIndexDataType);
+			}
+			check(InputTensors[Idx].HasPreparedData());
+			ConstantCPUInputs.Add(Idx);
+		}
+
+		// Initialize Input tensor desc
+        DmlUtil::FTensorDesc DmlInputTensorDesc;
+        if (!DmlInputTensorDesc.InitFromTensor(InputTensors[0], InputTensors[0].GetShape().Rank()))
+        {
+            UE_LOG(LogNNE, Error, TEXT("Failed to initialize Slice input for DML inference"));
+            return false;
+        }
+
+		// Initialize Output tensor desc
+        DmlUtil::FTensorDesc DmlOutputTensorDesc;
+        if (!DmlOutputTensorDesc.InitFromTensor(OutputTensors[0], OutputTensors[0].GetShape().Rank()))
+        {
+            UE_LOG(LogNNE, Error, TEXT("Failed to initialize Slice Output for DML inference"));
+            return false;
+        }
+
+
+		DmlUtil::FSmallUIntArray Offsets, Sizes;
+		DmlUtil::FSmallIntArray Strides;
+
+		switch(InputIndexDataType)
+		{
+		case ENNETensorDataType::Int32:
+			ComputeOffsetsSizesStrides<int32>(InputTensors, Offsets, Sizes, Strides);
+			break;
+		case ENNETensorDataType::Int64:
+			ComputeOffsetsSizesStrides<int64>(InputTensors, Offsets, Sizes, Strides);
+			break;
 		}
 
 
