Add RDG DirectML Gather operator

#rnx

diff --git a/Engine/Plugins/Experimental/NNE/Source/NNEQA/Resources/NNEQATestDesc.json b/Engine/Plugins/Experimental/NNE/Source/NNEQA/Resources/NNEQATestDesc.json
--- a/Engine/Plugins/Experimental/NNE/Source/NNEQA/Resources/NNEQATestDesc.json
+++ b/Engine/Plugins/Experimental/NNE/Source/NNEQA/Resources/NNEQATestDesc.json
@@ -328,6 +328,15 @@
 				{ "target": "Upsample" }
 			]
 		},
+		{
+			"category": "Gather",
+			"runtimes": [	
+				{ "name": "NNERuntimeRDGHlsl", "skip": true }
+			],
+			"targets": [
+				{ "target": "Gather" }
+			]
+		},
 		{
 			"category": "GRU",
 			"runtimes": [
@@ -702,6 +711,12 @@
 				{ "inputs":  [ { "shape": [ 1,32,224,128 ] } ], "outputs": [ { "shape": [ 1,33,227,135 ] } ] }
 			]
 		},
+		{
+			"name": "Gather",
+			"datasets": [
+				{ "inputs":  [ { "shape": [ 2, 3, 4 ] },  { "shape": [ 1 ], "source": [ "1" ], "type": "int32" } ], "outputs": [ { "shape": [ 1, 3, 4 ]} ] }
+			]
+		},
 		{
 			"name": "Upsample.ScaleFromInput",
 			"runtimes": [	
diff --git a/Engine/Plugins/Experimental/NNE/Source/NNERuntimeRDG/Private/Dml/NNEDmlModel.cpp b/Engine/Plugins/Experimental/NNE/Source/NNERuntimeRDG/Private/Dml/NNEDmlModel.cpp
--- a/Engine/Plugins/Experimental/NNE/Source/NNERuntimeRDG/Private/Dml/NNEDmlModel.cpp
+++ b/Engine/Plugins/Experimental/NNE/Source/NNERuntimeRDG/Private/Dml/NNEDmlModel.cpp
@@ -1130,7 +1130,8 @@ void FModel::AddDispatchOps_RenderThread(FRDGBuilder& GraphBuilder)
 			RHICmdList.EnqueueLambda(
 				[this](FRHICommandListImmediate& RHICmdList)
 				{
-					TArray<CD3DX12_RESOURCE_BARRIER, TInlineAllocator<MaxNumInputs + MaxNumOutputs>>	Barriers;
+					TArray<CD3DX12_RESOURCE_BARRIER, TInlineAllocator<MaxNumInputs + MaxNumOutputs>>	PreBarriers;
+					TArray<CD3DX12_RESOURCE_BARRIER, TInlineAllocator<MaxNumOutputs * 2>>				PostBarriers;
 
 					for (FRHIBuffer* Buffer : InputBuffers)
 					{
@@ -1142,7 +1143,7 @@ void FModel::AddDispatchOps_RenderThread(FRDGBuilder& GraphBuilder)
 						ID3D12Resource* Resource = DynamicRHI->RHIGetResource(Buffer);
 
 						//Note: We should not assume COPY_DEST state
-						Barriers.Emplace(
+						PreBarriers.Emplace(
 							CD3DX12_RESOURCE_BARRIER::Transition(
 								Resource,
 								D3D12_RESOURCE_STATE_COPY_DEST,
@@ -1150,6 +1151,30 @@ void FModel::AddDispatchOps_RenderThread(FRDGBuilder& GraphBuilder)
 						);
 					}
 
+					for (FRHIBuffer* Buffer : OutputBuffers)
+					{
+						ID3D12Resource* Resource = DynamicRHI->RHIGetResource(Buffer);
+
+						if ((Buffer->GetUsage() & EBufferUsageFlags::SourceCopy) == EBufferUsageFlags::SourceCopy)
+						{
+							PreBarriers.Emplace(
+								CD3DX12_RESOURCE_BARRIER::Transition(
+									Resource,
+									D3D12_RESOURCE_STATE_COMMON,
+									D3D12_RESOURCE_STATE_UNORDERED_ACCESS)
+							);
+
+							PostBarriers.Emplace(
+								CD3DX12_RESOURCE_BARRIER::Transition(
+									Resource,
+									D3D12_RESOURCE_STATE_UNORDERED_ACCESS,
+									D3D12_RESOURCE_STATE_COPY_SOURCE)
+							);
+						}
+
+						PostBarriers.Add(CD3DX12_RESOURCE_BARRIER::UAV(Resource));
+					}
+
 					//Note: We should use this instead of NNE_USE_D3D12_RESOURCES
 					//FBufferRHIRef TempBuff = MemSizeTemp ? CreateRHIBuffer(RHICmdList, MemSizeTemp, TempBuffUsage, TempBuffAccess, TEXT("FModel_Dispatch_TempBuff")) : nullptr;
 
@@ -1159,8 +1184,9 @@ void FModel::AddDispatchOps_RenderThread(FRDGBuilder& GraphBuilder)
 
 					D3DCmdList = DynamicRHI->RHIGetGraphicsCommandList(DevCtx->DeviceIndex);
 					D3DCmdList->SetDescriptorHeaps(1, &DescHeap);
-					D3DCmdList->ResourceBarrier(Barriers.Num(), Barriers.GetData());
+					D3DCmdList->ResourceBarrier(PreBarriers.Num(), PreBarriers.GetData());
 					DevCtx->CmdRec->RecordDispatch(D3DCmdList, CompiledOp, BindingTable->Get());
+					D3DCmdList->ResourceBarrier(PostBarriers.Num(), PostBarriers.GetData());
 
 					DynamicRHI->RHIFinishExternalComputeWork(DevCtx->DeviceIndex, D3DCmdList);
 				}
diff --git a/Engine/Plugins/Experimental/NNE/Source/NNERuntimeRDG/Private/Dml/NNEDmlOperator.h b/Engine/Plugins/Experimental/NNE/Source/NNERuntimeRDG/Private/Dml/NNEDmlOperator.h
--- a/Engine/Plugins/Experimental/NNE/Source/NNERuntimeRDG/Private/Dml/NNEDmlOperator.h
+++ b/Engine/Plugins/Experimental/NNE/Source/NNERuntimeRDG/Private/Dml/NNEDmlOperator.h
@@ -30,6 +30,12 @@ static constexpr uint32_t NcdhwDimensionCount = 5;
 static constexpr uint32_t NcdhwSpatialDimensionCount = 3;
 static constexpr uint32_t NonspatialDimensionCount = 2; // The batch and channel dimensions of NCW, NCHW, NCDHW....
 
+template<typename T>
+inline TArrayView<T> MakeEmptyArrayView()
+{
+	return MakeArrayView(static_cast<T*>(nullptr), 0);
+}
+
 //
 //
 //
diff --git a/Engine/Plugins/Experimental/NNE/Source/NNERuntimeRDG/Private/Dml/NNEDmlOperatorGather.cpp b/Engine/Plugins/Experimental/NNE/Source/NNERuntimeRDG/Private/Dml/NNEDmlOperatorGather.cpp
new file mode 100644
--- /dev/null
+++ b/Engine/Plugins/Experimental/NNE/Source/NNERuntimeRDG/Private/Dml/NNEDmlOperatorGather.cpp
@@ -0,0 +1,131 @@
+// Copyright Epic Games, Inc. All Rights Reserved.
+
+#ifdef NNE_USE_DIRECTML
+#include "NNEDmlOperator.h"
+
+namespace UE::NNERuntimeRDG::Private::Dml
+{
+
+inline int32 HandleNegativeAxis(int32 Axis, int32 Rank)
+{
+	if (Axis < 0)
+	{
+		Axis += Rank;
+		check(Axis < Rank);
+	}
+
+	return Axis;
+}
+
+class FOperatorDmlGather : public FOperatorDml
+{
+public:
+
+	static FOperatorDml* Create()
+	{
+		return new FOperatorDmlGather();
+	}
+
+	//
+	//
+	//
+	virtual bool Initialize(IDMLDevice* Device, TArrayView<const NNECore::Internal::FTensor> InputTensors, TArrayView<const NNECore::Internal::FTensor> OutputTensors, const NNECore::FAttributeMap& Attributes) override
+	{
+		check(InputTensors.Num() == 2);
+		check(OutputTensors.Num() == 1);
+		
+		const NNECore::Internal::FTensor& InputTensor = InputTensors[0];
+		const NNECore::Internal::FTensor& IndicesTensor = InputTensors[1];
+		const NNECore::Internal::FTensor& OutputTensor = OutputTensors[0];
+
+		const NNECore::FTensorShape& InputShape = InputTensor.GetShape();
+		const NNECore::FTensorShape& IndicesShape = IndicesTensor.GetShape();
+
+		if (InputShape.Rank() != IndicesShape.Rank())
+		{
+			UE_LOG(LogNNE, Warning, TEXT("Indices tensor rank must match input tensor rank"));
+			return false;
+		}
+
+		if (IndicesTensor.GetDataType() != ENNETensorDataType::UInt32 && IndicesTensor.GetDataType() != ENNETensorDataType::Int32)
+		{
+			UE_LOG(LogNNE, Warning, TEXT("DML only supports UINT32/INT32 for indices tensor"));
+			return false;
+		}
+
+		if (InputTensor.GetDataType() != OutputTensor.GetDataType())
+		{
+			UE_LOG(LogNNE, Warning, TEXT("Input and Output tensor should have same data type"));
+			return false;
+		}
+
+		// Read attributes
+		int32	Axis = Attributes.GetValueOrDefault<int>(TEXT("axis"), 0);
+
+		Axis = HandleNegativeAxis(Axis, InputTensor.GetShape().Rank());
+
+		// Compute output shape
+		const int32 OutputRank = IndicesShape.Rank() + InputShape.Rank() - 1;
+		DmlUtil::FSmallUIntArray OutputShape;
+		int32 DataRankIdx = 0;
+
+		for (; DataRankIdx < Axis; ++DataRankIdx)
+		{
+			OutputShape.Add(InputShape.GetData()[DataRankIdx]);
+		}
+
+		OutputShape.Append(IndicesShape.GetData());
+		++DataRankIdx;
+
+		for (; DataRankIdx < InputShape.Rank(); ++DataRankIdx)
+		{
+			OutputShape.Add(InputShape.GetData()[DataRankIdx]);
+		}
+
+		if (InputShape.Rank() != OutputRank)
+		{
+			UE_LOG(LogNNE, Warning, TEXT("Output tensor rank must match input tensor rank"));
+			return false;
+		}
+
+		// Initialize tensor descriptors
+		DmlUtil::FTensorDesc	DmlInputTensor{};
+		DmlUtil::FTensorDesc	DmlIndicesTensor{};
+		DmlUtil::FTensorDesc	DmlOutputTensor{};
+
+		if (!DmlInputTensor.InitFromTensor(InputTensor, 1))
+		{
+			UE_LOG(LogNNE, Warning, TEXT("Failed to initialize tensor(s) for DML inference"));
+			return false;
+		}
+
+		if (!DmlIndicesTensor.InitFromTensor(IndicesTensor, InputShape.Rank()))
+		{
+			UE_LOG(LogNNE, Warning, TEXT("Failed to initialize tensor(s) for DML inference"));
+			return false;
+		}
+
+		if (!DmlOutputTensor.InitFromTensor(OutputTensor, InputShape.Rank(), MakeEmptyArrayView<uint32>(), OutputShape))
+		{
+			UE_LOG(LogNNE, Warning, TEXT("Failed to initialize tensor(s) for DML inference"));
+			return false;
+		}
+
+		DML_GATHER_OPERATOR_DESC	OpDesc{};
+
+		OpDesc.InputTensor = &DmlInputTensor.Desc;
+		OpDesc.IndicesTensor = &DmlIndicesTensor.Desc;
+		OpDesc.OutputTensor = &DmlOutputTensor.Desc;
+		OpDesc.Axis = Axis;
+		OpDesc.IndexDimensions = IndicesTensor.GetShape().Rank();
+
+		return CreateOperator(Device, DML_OPERATOR_DESC { DML_OPERATOR_GATHER, &OpDesc });
+	}
+};
+
+// Register operator on Module startup
+NNE_DML_REGISTER_OP(Gather)
+
+} // namespace UE::NNERuntimeRDG::Private::Dml
+
+#endif // NNE_USE_DIRECTML
