NNE - Fix Dml Reshape handling of -1 dimension value in shape tensor
#rnx

diff --git a/Engine/Plugins/Experimental/NNE/Source/NNEQA/Resources/NNEQATestDesc.json b/Engine/Plugins/Experimental/NNE/Source/NNEQA/Resources/NNEQATestDesc.json
--- a/Engine/Plugins/Experimental/NNE/Source/NNEQA/Resources/NNEQATestDesc.json
+++ b/Engine/Plugins/Experimental/NNE/Source/NNEQA/Resources/NNEQATestDesc.json
@@ -914,6 +914,20 @@
 			"datasets": [
 				{ "inputs": [ { "shape": [ 2, 3, 4 ] } ], "weights": [ { "shape": [ 1 ], "source": [ "24" ], "type": "int64" } ], "outputs": [ { "shape": [ 24 ] } ] }
 			]
+		},
+		{
+			"name": "Reshape.MinusOne",
+			"datasets": [
+				{ "inputs": [ { "shape": [ 2, 4, 3 ] } ], "weights": [ { "shape": [ 2 ], "source": [ "12", "-1" ], "type": "int64" } ], "outputs": [ { "shape": [ 12, 2 ] } ], 
+				  "runtimes": [ { "name": "NNERuntimeRDGHlsl", "skip": true, "comment": "Shape inference here not yet supported" } ] }
+			]
+		},
+		{
+			"name": "Reshape.Zero",
+			"datasets": [
+				{ "inputs": [ { "shape": [ 2, 4, 3 ] } ], "weights": [ { "shape": [ 3 ], "source": [ "3", "0", "2" ], "type": "int64" } ], "outputs": [ { "shape": [ 3, 4, 2 ] } ], 
+				  "runtimes": [ { "name": "NNERuntimeRDGHlsl", "skip": true, "comment": "Shape inference here not yet supported" } ] }
+			]
 		}
 	],
 
diff --git a/Engine/Plugins/Experimental/NNE/Source/NNERuntimeRDG/Private/Dml/NNEDmlOperatorReshape.cpp b/Engine/Plugins/Experimental/NNE/Source/NNERuntimeRDG/Private/Dml/NNEDmlOperatorReshape.cpp
--- a/Engine/Plugins/Experimental/NNE/Source/NNERuntimeRDG/Private/Dml/NNEDmlOperatorReshape.cpp
+++ b/Engine/Plugins/Experimental/NNE/Source/NNERuntimeRDG/Private/Dml/NNEDmlOperatorReshape.cpp
@@ -4,6 +4,7 @@
 
 #include "NNEDmlOperator.h"
 #include "Algo/Transform.h"
+#include "Algo/Count.h"
 #include "Misc/EnumerateRange.h"
 
 namespace UE::NNERuntimeRDG::Private::Dml
@@ -14,12 +15,71 @@ namespace UE::NNERuntimeRDG::Private::Dml
  */
 class FOperatorDmlReshape : public FOperatorDml
 {
+
     template<typename DataType>
-    DmlUtil::FSmallUIntArray TensorContentToUIntArray(const NNECore::Internal::FTensor& InputTensor)
+    bool ReshapeTensorToDmlShape(
+        const NNECore::Internal::FTensor& InputTensor, 
+        const NNECore::Internal::FTensor& ShapeTensor, 
+        bool bAllowZero,
+        DmlUtil::FSmallUIntArray& OutShape)
     {
-        DmlUtil::FSmallUIntArray OutArray;
-        Algo::Transform(InputTensor.GetPreparedData<DataType>(), OutArray, [](DataType In){ return (uint32) In; });
-        return OutArray;
+        
+        DmlUtil::FSmallArray<DataType> ReshapedShape(ShapeTensor.GetPreparedData<DataType>());
+
+        if(!bAllowZero)
+        {
+            // at most 1 dimension can be -1
+            if(Algo::Count(ReshapedShape, -1) > 1)
+            {
+                UE_LOG(LogNNE, Error, TEXT("Shape tensor can't contain more than one '-1'."));
+                return false;
+            }
+            for (TEnumerateRef<DataType> Elem : EnumerateRange(ReshapedShape))
+            {
+                if(!(*Elem != 0 || InputTensor.GetShape().Rank() > Elem.GetIndex()))
+                {
+                    UE_LOG(LogNNE, Error, TEXT("Shape tensor contains '0' in an invalid place."));
+                    return false;
+                }
+                *Elem = 
+                    *Elem == 0 ? 
+                        InputTensor.GetShape().GetData()[Elem.GetIndex()] 
+                        : 
+                        *Elem;
+            }
+        }
+        else
+        {
+            // no -1 is allowed if there is a 0
+            if(!Algo::Count(ReshapedShape, 0) == 0 && !Algo::Count(ReshapedShape, -1) == 0)
+            {
+                UE_LOG(LogNNE, Error, TEXT("Shape tensor contains both '0' and '-1'. This is not allowed."));
+                return false;
+            }
+        }
+
+        auto PartialVolume = [](TConstArrayView<DataType> Shape) -> uint32
+            {
+                uint32 Volume = 1;
+                for(DataType Dim : Shape)
+                {
+                    Volume *= Dim != -1 ? (DataType) Dim : 1;
+                }
+                return Volume;
+            };
+
+        OutShape.SetNumUninitialized(ReshapedShape.Num());
+
+        for(TConstEnumerateRef<DataType> Elem : EnumerateRange(ReshapedShape))
+        {
+            OutShape[Elem.GetIndex()] = 
+                *Elem == -1 ? 
+                    (uint32) InputTensor.GetShape().Volume() / PartialVolume(ReshapedShape)
+                    :
+                    (uint32) *Elem;
+        }
+
+        return true;
     }
 
 public:
@@ -42,36 +102,37 @@ public:
         // Shape tensor must be constant!
         check(InputTensors[1].HasPreparedData());
 
+        bool bAllowZero = (bool)
+			( Attributes.GetValueOrDefault<int32>(TEXT("allowzero"), 0) );
+
         DmlUtil::FSmallUIntArray ReshapedShape;
 
         switch(InputTensors[1].GetDataType())
         {
         case ENNETensorDataType::Int32:
-            ReshapedShape = TensorContentToUIntArray<int32>(InputTensors[1]);
+            if(!ReshapeTensorToDmlShape<int32>(InputTensors[0], InputTensors[1], bAllowZero, ReshapedShape))
+            {
+                return false;
+            }
             break;
         case ENNETensorDataType::Int64:
-            ReshapedShape = TensorContentToUIntArray<int64>(InputTensors[1]);
+            if(!ReshapeTensorToDmlShape<int64>(InputTensors[0], InputTensors[1], bAllowZero, ReshapedShape))
+            {
+                return false;
+            }
             break;
         case ENNETensorDataType::UInt32:
-            ReshapedShape = TensorContentToUIntArray<uint32>(InputTensors[1]);
+            if(!ReshapeTensorToDmlShape<uint32>(InputTensors[0], InputTensors[1], bAllowZero, ReshapedShape))
+            {
+                return false;
+            }
             break;
         default:
             UE_LOG(LogNNE, Warning, TEXT("Shape tensor has invalid data type"));
 			return false;
         };
         
-
-        bool bAllowZero = (bool)
-			( Attributes.GetValueOrDefault<int32>(TEXT("allowzero"), 0) );
-
-        if(!bAllowZero)
-        {
-            for (TEnumerateRef<uint32> Elem : EnumerateRange(ReshapedShape))
-            {
-                check(*Elem != 0 || InputTensors[0].GetShape().Rank() > Elem.GetIndex());
-                *Elem = *Elem == 0 ? InputTensors[0].GetShape().GetData()[Elem.GetIndex()] : *Elem;
-            }
-        }
+        check(ReshapedShape == OutputTensors[0].GetShape().GetData());
 
         DmlUtil::FTensorDesc DmlInputTensorDesc;
         if (!DmlInputTensorDesc.InitFromTensor(InputTensors[0], ReshapedShape.Num(), 
