Fixed tests failures for RDG DML operators

#rnx

diff --git a/Engine/Plugins/Experimental/NNE/Source/NNEQA/Resources/NNEQATestDesc.json b/Engine/Plugins/Experimental/NNE/Source/NNEQA/Resources/NNEQATestDesc.json
--- a/Engine/Plugins/Experimental/NNE/Source/NNEQA/Resources/NNEQATestDesc.json
+++ b/Engine/Plugins/Experimental/NNE/Source/NNEQA/Resources/NNEQATestDesc.json
@@ -760,6 +760,7 @@
 				{ 
 					"inputs": [ { "shape": [ 2, 2, 3 ] }, { "shape": [ 2 ] }, { "shape": [ 2 ] } ], "outputs": [ { "shape": [ 2, 2, 3 ] } ], 
 					"runtimes": [
+						{ "name": "NNERuntimeRDGDml", "skip": true, "comment": "Not supported" },
 						{ "name": "NNERuntimeORTDml", "skip": true, "comment": "Not supported" } 
 					]
 				},
diff --git a/Engine/Plugins/Experimental/NNE/Source/NNERuntimeRDG/Private/Dml/NNEDmlOperatorConvolution.cpp b/Engine/Plugins/Experimental/NNE/Source/NNERuntimeRDG/Private/Dml/NNEDmlOperatorConvolution.cpp
--- a/Engine/Plugins/Experimental/NNE/Source/NNERuntimeRDG/Private/Dml/NNEDmlOperatorConvolution.cpp
+++ b/Engine/Plugins/Experimental/NNE/Source/NNERuntimeRDG/Private/Dml/NNEDmlOperatorConvolution.cpp
@@ -103,6 +103,26 @@ class FOperatorDmlConv : public FOperatorDml
 				WindowSize.Add(FilterShape.GetData()[Dim]);
 			}
 
+			if (Direction == DML_CONVOLUTION_DIRECTION_FORWARD)
+			{
+				OutPadding.Init(0, NumDimensions);
+			}
+			else
+			{
+				const FNNEAttributeValue* AttrOutPadding = Attributes.GetAttributeValue(TEXT("output_padding"));
+
+				if (AttrOutPadding)
+				{
+					OutPadding = AttrOutPadding->GetValue<TArray<int32>>();
+				}
+				else
+				{
+					OutPadding.Init(0, NumDimensions);
+				}
+			}
+
+			Group = Attributes.GetValueOrDefault<int32>(TEXT("group"), 1);
+
 			AutoPad = AutoPadFromString(*Attributes.GetValue<FString>(TEXT("auto_pad")));
 
 			if (AutoPad == EAutoPad::NOTSET)
@@ -138,26 +158,6 @@ class FOperatorDmlConv : public FOperatorDml
 				SetAutoPadding(InputShape.GetData(), FilterShape.GetData());
 			}
 
-			Group = Attributes.GetValueOrDefault<int32>(TEXT("group"), 1);
-
-			if (Direction == DML_CONVOLUTION_DIRECTION_FORWARD)
-			{
-				OutPadding.Init(0, NumDimensions);
-			}
-			else
-			{
-				const FNNEAttributeValue* AttrOutPadding = Attributes.GetAttributeValue(TEXT("output_padding"));
-
-				if (AttrOutPadding)
-				{
-					OutPadding = AttrOutPadding->GetValue<TArray<int32>>();
-				}
-				else
-				{
-					OutPadding.Init(0, NumDimensions);
-				}
-			}
-
 			if (Direction == DML_CONVOLUTION_DIRECTION_FORWARD)
 			{
 				SetOutputShape(InputShape.GetData(), FilterShape.GetData());
@@ -193,11 +193,21 @@ class FOperatorDmlConv : public FOperatorDml
 
 			for (uint32 Dim = 0; Dim < NumDimensions; ++Dim)
 			{
-				uint32 InputLen = uint32(InputShape[Dim + DimOffset]);
-				uint32 StridedOutLen = (InputLen + Strides[Dim] - 1) / Strides[Dim];
-				uint32 KernelLen = 1 + (WindowSize[Dim] - 1) * Dilations[Dim];
-				uint32 Len = Strides[Dim] * (StridedOutLen - 1) + KernelLen;
-				uint32 Padding = (Len <= InputLen) ? 0 : (Len - InputLen);
+				uint32 Padding;
+
+				if (Direction == DML_CONVOLUTION_DIRECTION_FORWARD)
+				{
+					uint32 InputLen = uint32(InputShape[Dim + DimOffset]);
+					uint32 StridedOutLen = (InputLen + Strides[Dim] - 1) / Strides[Dim];
+					uint32 KernelLen = 1 + (WindowSize[Dim] - 1) * Dilations[Dim];
+					uint32 Len = Strides[Dim] * (StridedOutLen - 1) + KernelLen;
+					
+					Padding = (Len <= InputLen) ? 0 : (Len - InputLen);
+				}
+				else
+				{
+					Padding = (InputShape[Dim + DimOffset] - 1) * Dilations[Dim] - Strides[Dim] + OutPadding[Dim] + 1;
+				}
 
 				if (AutoPad == EAutoPad::SAME_LOWER)
 				{
diff --git a/Engine/Plugins/Experimental/NNE/Source/NNERuntimeRDG/Private/Dml/NNEDmlOperatorGather.cpp b/Engine/Plugins/Experimental/NNE/Source/NNERuntimeRDG/Private/Dml/NNEDmlOperatorGather.cpp
--- a/Engine/Plugins/Experimental/NNE/Source/NNERuntimeRDG/Private/Dml/NNEDmlOperatorGather.cpp
+++ b/Engine/Plugins/Experimental/NNE/Source/NNERuntimeRDG/Private/Dml/NNEDmlOperatorGather.cpp
@@ -41,7 +41,7 @@ public:
 		const NNECore::FTensorShape& InputShape = InputTensor.GetShape();
 		const NNECore::FTensorShape& IndicesShape = IndicesTensor.GetShape();
 
-		if (InputShape.Rank() != IndicesShape.Rank())
+		if (IndicesShape.Rank() > InputShape.Rank())
 		{
 			UE_LOG(LogNNE, Warning, TEXT("Indices tensor rank must match input tensor rank"));
 			return false;
