NNE Retrained NNs for style transfer demo

#rnx

diff --git a/Engine/Plugins/Experimental/NNE/Source/NNEQA/Resources/NNEQATestDesc.json b/Engine/Plugins/Experimental/NNE/Source/NNEQA/Resources/NNEQATestDesc.json
--- a/Engine/Plugins/Experimental/NNE/Source/NNEQA/Resources/NNEQATestDesc.json
+++ b/Engine/Plugins/Experimental/NNE/Source/NNEQA/Resources/NNEQATestDesc.json
@@ -6,7 +6,7 @@
 			"is_model_test": true,
 			"targets": [
 				{
-					"target": "Mosaic-9",
+					"target": "udnie-20-small-720p",
 					"runtimes": [
 						{ "name": "NNERuntimeRDGDml",  "skip": true, "comment": "Model not supported yet." }
 					],
@@ -19,7 +19,7 @@
 			"is_model_test": true,
 			"targets": [
 				{
-					"target": "Mosaic-9-dyn",
+					"target": "udnie-20-small-dyn",
 					"runtimes": [
 						{ "name": "NNERuntimeRDGDml",  "skip": true, "comment": "Model not supported yet." }
 					],
