Fix error: logical not is only applied to the left hand side of this comparison
#rnx

diff --git a/Engine/Plugins/Experimental/NNE/Source/NNERuntimeRDG/Private/Dml/NNEDmlOperatorReshape.cpp b/Engine/Plugins/Experimental/NNE/Source/NNERuntimeRDG/Private/Dml/NNEDmlOperatorReshape.cpp
--- a/Engine/Plugins/Experimental/NNE/Source/NNERuntimeRDG/Private/Dml/NNEDmlOperatorReshape.cpp
+++ b/Engine/Plugins/Experimental/NNE/Source/NNERuntimeRDG/Private/Dml/NNEDmlOperatorReshape.cpp
@@ -51,7 +51,7 @@ class FOperatorDmlReshape : public FOperatorDml
         else
         {
             // no -1 is allowed if there is a 0
-            if(!Algo::Count(ReshapedShape, 0) == 0 && !Algo::Count(ReshapedShape, -1) == 0)
+            if(Algo::Count(ReshapedShape, 0) != 0 && Algo::Count(ReshapedShape, -1) != 0)
             {
                 UE_LOG(LogNNE, Error, TEXT("Shape tensor contains both '0' and '-1'. This is not allowed."));
                 return false;
