NNE - Fix wrong weight tensor indexing in NNEDmlModel.cpp
#rnx

diff --git a/Engine/Plugins/Experimental/NNE/Source/NNERuntimeRDG/Private/Dml/NNEDmlModel.cpp b/Engine/Plugins/Experimental/NNE/Source/NNERuntimeRDG/Private/Dml/NNEDmlModel.cpp
--- a/Engine/Plugins/Experimental/NNE/Source/NNERuntimeRDG/Private/Dml/NNEDmlModel.cpp
+++ b/Engine/Plugins/Experimental/NNE/Source/NNERuntimeRDG/Private/Dml/NNEDmlModel.cpp
@@ -958,23 +958,22 @@ bool FModel::InitCompiledOp(TConstArrayView<int32> OpInputIndices, uint64 Tensor
 			if (TensorDataSize)
 			{
 				UploadFence = RHICreateGPUFence(TEXT("FInferenceModel_UploadFence"));
-				
 				FRHIBuffer*		UploadBuff = CreateRHIBuffer(RHICmdList, TensorDataSize, BUF_ShaderResource | BUF_Dynamic | BUF_FastVRAM, ERHIAccess::CopySrc, TEXT("FInferenceModel_UploadBuffer"));
 				uint8*			UploadBuffPtr = static_cast<uint8*>(RHICmdList.LockBuffer(UploadBuff, 0, TensorDataSize, RLM_WriteOnly_NoOverwrite));
 				uint64			UploadOffset = 0;
 
-				for (int32 WeightIdx : WeightTensorIndices)
+				int32 WeightIdx = 0;
+				for (int32 TensorIdx : WeightTensorIndices)
 				{
-					if (ConstantCPUTensorIndices.Find(WeightIdx) != -1)
+					if (ConstantCPUTensorIndices.Find(TensorIdx) != -1)
 					{
 						continue;
 					}
-
-					const FTensorRDG&		Tensor = WeightTensorRDGs[WeightIdx];
+					
+					const FTensorRDG&		Tensor = WeightTensorRDGs[WeightIdx++];
 					TConstArrayView<uint8>	TensorData = Tensor.GetPreparedData<uint8>();
 
 					FBufferRHIRef WeightBuff;
-
 					WeightBuff = CreateRHIBuffer(RHICmdList, TensorData.Num(), WeightBuffUsage, WeightBuffAccess, TEXT("FModel_TensorWeights"));
 				
 					FMemory::Memcpy(UploadBuffPtr + UploadOffset, TensorData.GetData(), TensorData.Num());
