NNE - Dml GRU cast SequenceLenghtsTensor and fix wrong handling of FTensor.
#rnx

diff --git a/Engine/Plugins/Experimental/NNE/Source/NNERuntimeRDG/Private/Dml/NNEDmlOperatorGRU.cpp b/Engine/Plugins/Experimental/NNE/Source/NNERuntimeRDG/Private/Dml/NNEDmlOperatorGRU.cpp
--- a/Engine/Plugins/Experimental/NNE/Source/NNERuntimeRDG/Private/Dml/NNEDmlOperatorGRU.cpp
+++ b/Engine/Plugins/Experimental/NNE/Source/NNERuntimeRDG/Private/Dml/NNEDmlOperatorGRU.cpp
@@ -26,15 +26,6 @@ union FDMLGRUActivationOpDescUnion
  */
 class FOperatorDmlGRU : public FOperatorDml
 {
-	static NNECore::Internal::FTensor To4DTensor(const NNECore::Internal::FTensor& InTensor)
-	{
-		check(InTensor.GetShape().Rank() <= 4);
-		TArray<uint32, TInlineAllocator<4>> NewShapeArray;
-		NewShapeArray.Init(1, 4 - InTensor.GetShape().Rank());
-		NewShapeArray.Append(InTensor.GetShape().GetData());
-		
-		return NNECore::Internal::FTensor::Make(InTensor.GetName(), NNECore::FTensorShape::Make(NewShapeArray), InTensor.GetDataType());
-	}
 
 	static DML_RECURRENT_NETWORK_DIRECTION DirectionFromString(FStringView StringVal)
 	{
@@ -144,8 +135,7 @@ public:
 		DmlUtil::FTensorDesc	Dml##Name{};\
 		if(Cond)\
 		{\
-			NNECore::Internal::FTensor Name = To4DTensor(Tensor);\
-			if (!InitDmlTensorDesc(Dml##Name, Name))\
+			if (!Dml##Name.InitFromTensor(Tensor, 4))\
 			{\
 				UE_LOG(LogNNE, Error, TEXT("Failed to initialize GRU's "#Name" for DML inference"));\
 				return false;\
@@ -206,6 +196,7 @@ public:
 		if(InputTensors.Num() >= 5)
 		{
 			check(BatchSize == InputTensors[4].GetShape().GetData()[0]);
+			check(InputTensors[4].GetDataType() == ENNETensorDataType::Int32);
 		}
 		if(InputTensors.Num() >= 6)
 		{
@@ -234,6 +225,8 @@ public:
 		SET_DMLTENSORDESC_FROM_TENSOR(DmlGRUOpDesc, InputTensors[2], RecurrenceTensor)
 		SET_DMLTENSORDESC_FROM_TENSOR_COND(DmlGRUOpDesc, InputTensors[3], BiasTensor, InputTensors.Num() >= 4)
 		SET_DMLTENSORDESC_FROM_TENSOR_COND(DmlGRUOpDesc, InputTensors[4], SequenceLengthsTensor, InputTensors.Num() >= 5)
+		// Cast SequenceLengthsTensor from int32 to uint32 due to differences in representation between ONNX and DML formats.
+		DmlSequenceLengthsTensor.BuffDesc.DataType = DML_TENSOR_DATA_TYPE::DML_TENSOR_DATA_TYPE_UINT32;
 		SET_DMLTENSORDESC_FROM_TENSOR_COND(DmlGRUOpDesc, InputTensors[5], HiddenInitTensor, InputTensors.Num() >= 6)
 		SET_DMLTENSORDESC_FROM_TENSOR_COND(DmlGRUOpDesc, OutputTensors[0], OutputSequenceTensor, OutputTensors.Num() >= 1)
 		SET_DMLTENSORDESC_FROM_TENSOR_COND(DmlGRUOpDesc, OutputTensors[1], OutputSingleTensor, OutputTensors.Num() >= 2)
