DML graph builder now supports inputs and initializers that are connected to multiple nodes

#rnx
#okforgibhub public, internal

diff --git a/Engine/Plugins/Experimental/NNE/Source/NNERuntimeRDG/Private/Dml/NNEDmlModel.cpp b/Engine/Plugins/Experimental/NNE/Source/NNERuntimeRDG/Private/Dml/NNEDmlModel.cpp
--- a/Engine/Plugins/Experimental/NNE/Source/NNERuntimeRDG/Private/Dml/NNEDmlModel.cpp
+++ b/Engine/Plugins/Experimental/NNE/Source/NNERuntimeRDG/Private/Dml/NNEDmlModel.cpp
@@ -539,60 +539,98 @@ private:
 		NumInputs = 0;
 		NumOutputs = 0;
 
-		for (int32 Idx = 0; Idx < InGraph.InputIndices.Num(); ++Idx)
+		// Loop through all operators that have this tensor as an input/weight and it as an input edge
+		auto AddInputs = [this, InGraph](TConstArrayView<int32> InTensorIndices) -> bool
 		{
-			const int32 TensorIdx = InGraph.InputIndices[Idx];
-			AddInput(TensorIdx);
-		}
-
-		for (int32 Idx = 0; Idx < InGraph.WeightIndices.Num(); ++Idx)
-		{
-			if (InGraph.ConstantCPUIndices.Find(InGraph.WeightIndices[Idx]) == -1)
+			for (int32 Idx = 0; Idx < InTensorIndices.Num(); ++Idx)
 			{
-				AddInput(InGraph.WeightIndices[Idx]);
+				const int32 TensorIdx = InTensorIndices[Idx];
+
+				for (const FOpDesc& OpDesc : InGraph.Operators)
+				{
+					for (int32 InIdx = 0; InIdx < OpDesc.InputCount; ++InIdx)
+					{
+						const int32 OpInputIdx = InGraph.OpInputIndices[InIdx + OpDesc.InputStart];
+
+						if (OpInputIdx == TensorIdx)
+						{
+							bool bRes = AddInput(TensorIdx);
+							if (!bRes)
+							{
+								return bRes;
+							}
+						}
+					}
+				}
 			}
+
+			return true;
+		};
+
+		if (!AddInputs(InGraph.InputIndices))
+		{
+			UE_LOG(LogNNE, Warning, TEXT("DMLGraphBuilder failed to process input tensors"));
+			return false;
+		}
+		
+		if (!AddInputs(InGraph.WeightIndices))
+		{
+			UE_LOG(LogNNE, Warning, TEXT("DMLGraphBuilder failed to process weight tensors"));
+			return false;
 		}
 
 		for (int32 Idx = 0; Idx < InGraph.OutputIndices.Num(); ++Idx)
 		{
 			const int32 TensorIdx = InGraph.OutputIndices[Idx];
-			AddOutput(TensorIdx);
+			
+			if (!AddOutput(TensorIdx))
+			{
+				UE_LOG(LogNNE, Warning, TEXT("DMLGraphBuilder failed to process output tensors"));
+				return false;
+			}
 		}
 
 		Operators.Reset(InGraph.Operators.Num());
 		for (const FOpDesc& OpDesc : InGraph.Operators)
 		{
-			AddOp(OpDesc, InGraph);
+			if (!AddOp(OpDesc, InGraph))
+			{
+				UE_LOG(LogNNE, Warning, TEXT("DMLGraphBuilder failed to process operator tensors"));
+				return false;
+			}
 		}
 
-		//Note: Validate edges here
+		//Note: TODO: Validate edges here
 		
 		return true;
 	}
 
-	void AddInput(int32 TensorIdx)
+	bool AddInput(int32 TensorIdx)
 	{
-		AddEdge(
+		bool bRes = AddEdge(
 			FEdge(EEdgeType::Input)
 				.SetTensorIdx(TensorIdx)
-				.SetNodeSrcOutput(NumInputs)
+				.SetNodeSrcOutput(NumInputs),
+			false
 		);
 
 		++NumInputs;
+		return bRes;
 	}
 
-	void AddOutput(int32 TensorIdx)
+	bool AddOutput(int32 TensorIdx)
 	{
-		AddEdge(
+		bool bRes = AddEdge(
 			FEdge(EEdgeType::Output)
 				.SetTensorIdx(TensorIdx)
 				.SetNodeDstInput(NumOutputs)
 		);
 
 		++NumOutputs;
+		return bRes;
 	}
 
-	void AddIntermediate(int32 TensorIdx, int32 NodeSrc, int32 NodeSrcOutput)
+	bool AddIntermediate(int32 TensorIdx, int32 NodeSrc, int32 NodeSrcOutput)
 	{
 		FEdge* ConnEdge = Edges.FindByPredicate(
 				[TensorIdx](const FEdge& Curr)
@@ -603,11 +641,11 @@ private:
 
 		if (ConnEdge)
 		{
-			ConnectEdgeSrc(TensorIdx, NodeSrc, NodeSrcOutput);
+			return ConnectEdgeSrc(TensorIdx, NodeSrc, NodeSrcOutput);
 		}
 		else
 		{
-			AddEdge(
+			return AddEdge(
 				FEdge(EEdgeType::Intermediate)
 					.SetTensorIdx(TensorIdx)
 					.SetNodeSrc(NodeSrc)
@@ -616,7 +654,7 @@ private:
 		}
 	}
 
-	void AddEdge(const FEdge& Edge)
+	bool AddEdge(const FEdge& Edge, bool bIsUnique = true)
 	{
 		FEdge* StartEdge = Edges.FindByPredicate(
 				[Edge](const FEdge& Curr)
@@ -625,8 +663,17 @@ private:
 				}
 		);
 
-		check(StartEdge == nullptr);
+		if (bIsUnique)
+		{
+			if (StartEdge != nullptr)
+			{
+				UE_LOG(LogNNE, Warning, TEXT("DmlGraphBuilder:Edge is not unique"));
+				return false;
+			}
+		}
+
 		Edges.Add(Edge);
+		return true;
 	}
 
 	bool ConnectEdgeDst(int32 TensorIdx, int32 NodeDst, int32 NodeDstInput)
@@ -692,7 +739,7 @@ private:
 		return bFoundEdge;
 	}
 
-	void AddOp(const FOpDesc& InOp, const FGraphDesc& InGraph)
+	bool AddOp(const FOpDesc& InOp, const FGraphDesc& InGraph)
 	{
 		DML_OPERATOR_GRAPH_NODE_DESC& OpDesc = Operators.Add_GetRef({});
 
@@ -705,15 +752,25 @@ private:
 		{
 			const int32 TensorIdx = InGraph.OpInputIndices[Idx + InOp.InputStart];				
 			
-			ConnectEdgeDst(TensorIdx, NodeIdx, Idx);
+			bool bRes = ConnectEdgeDst(TensorIdx, NodeIdx, Idx);
+			if (!bRes)
+			{
+				return bRes;
+			}
 		}
 
 		for (int32 Idx = 0; Idx < InOp.OutputCount; ++Idx)
 		{
 			const int32 TensorIdx = InGraph.OpOutputIndices[Idx + InOp.OutputStart];
 
-			AddIntermediate(TensorIdx, NodeIdx, Idx);
+			bool bRes = AddIntermediate(TensorIdx, NodeIdx, Idx);
+			if (!bRes)
+			{
+				return bRes;
+			}
 		}
+
+		return true;
 	}
 
 	TArray<FEdge>							Edges;
