Fix RDG DML model input buffers

#rnx

diff --git a/Engine/Plugins/Experimental/NNE/Source/NNEQA/Resources/NNEQATestDesc.json b/Engine/Plugins/Experimental/NNE/Source/NNEQA/Resources/NNEQATestDesc.json
--- a/Engine/Plugins/Experimental/NNE/Source/NNEQA/Resources/NNEQATestDesc.json
+++ b/Engine/Plugins/Experimental/NNE/Source/NNEQA/Resources/NNEQATestDesc.json
@@ -13,8 +13,8 @@
 					"absolute_tolerance": 5.0e-2
 				}
 			]
-				},
-				{
+		},
+		{
 			"category": "StyleTransfer.VariableInputSize",
 			"is_model_test": true,
 			"targets": [
@@ -35,9 +35,6 @@
 				{
 					"target": "NeuralMorphModel_global",
 					"tags": [ "SmokeTest" ],
-					"runtimes": [
-						{ "name": "NNERuntimeRDGDml",  "skip": true, "comment": "Model not supported yet." }
-					],
 					"absolute_tolerance": 1e-5
 				},
 				{
@@ -60,11 +57,7 @@
 				},
 				{
 					"target": "VertexDeltaModel",
-					"runtimes": [
-						{ "name": "NNERuntimeORTDml", "absolute_tolerance": 1e-4 },
-						{ "name": "NNERuntimeRDGHlsl", "absolute_tolerance": 1e-4 },
-						{ "name": "NNERuntimeRDGDml",  "skip": true, "comment": "Model not supported yet." }
-					]
+					"absolute_tolerance": 1e-4
 				}
 			]
 		},
diff --git a/Engine/Plugins/Experimental/NNE/Source/NNERuntimeRDG/Private/Dml/NNEDmlModel.cpp b/Engine/Plugins/Experimental/NNE/Source/NNERuntimeRDG/Private/Dml/NNEDmlModel.cpp
--- a/Engine/Plugins/Experimental/NNE/Source/NNERuntimeRDG/Private/Dml/NNEDmlModel.cpp
+++ b/Engine/Plugins/Experimental/NNE/Source/NNERuntimeRDG/Private/Dml/NNEDmlModel.cpp
@@ -1154,11 +1154,13 @@ void FModel::AddDispatchOps_RenderThread(FRDGBuilder& GraphBuilder)
 		);
 	}
 
+	int32 NumWeightTensors = 0;
+
 	for (int32 Idx = 0; Idx < WeightTensorIndices.Num(); ++Idx)
 	{
 		if (ConstantCPUTensorIndices.Find(WeightTensorIndices[Idx]) == -1)
 		{
-			InputBuffers.Emplace(nullptr);
+			++NumWeightTensors;
 		}
 	}
 
@@ -1181,14 +1183,19 @@ void FModel::AddDispatchOps_RenderThread(FRDGBuilder& GraphBuilder)
 	GraphBuilder.AddPass(
 		RDG_EVENT_NAME("FModel_Dispatch"),
 		ERDGPassFlags::None | ERDGPassFlags::NeverCull,
-		[this](FRHICommandListImmediate& RHICmdList)
+		[this, NumWeightTensors](FRHICommandListImmediate& RHICmdList)
 		{
 			RHICmdList.EnqueueLambda(
-				[this](FRHICommandListImmediate& RHICmdList)
+				[this, NumWeightTensors](FRHICommandListImmediate& RHICmdList)
 				{
 					TArray<CD3DX12_RESOURCE_BARRIER, TInlineAllocator<MaxNumInputs + MaxNumOutputs>>	PreBarriers;
 					TArray<CD3DX12_RESOURCE_BARRIER, TInlineAllocator<MaxNumOutputs * 2>>				PostBarriers;
 
+					for (int32 Idx = 0; Idx < NumWeightTensors; ++Idx)
+					{
+						InputBuffers.Emplace(nullptr);
+					}
+
 					for (FRHIBuffer* Buffer : InputBuffers)
 					{
 						if (!Buffer)
