NNE HLSL more ops: Clip, Flatten, BatchNormalization, Dropout, Identity + Fixes for Reshape ops +Improve test coverage for Elu/Selu&HardSigmoid.

Tests:
* All tests green on win64 editor
* Style transfer demo working in editor

#rnx

diff --git a/Engine/Plugins/Experimental/NNE/Source/NNECore/Internal/NNECoreOperator.h b/Engine/Plugins/Experimental/NNE/Source/NNECore/Internal/NNECoreOperator.h
--- a/Engine/Plugins/Experimental/NNE/Source/NNECore/Internal/NNECoreOperator.h
+++ b/Engine/Plugins/Experimental/NNE/Source/NNECore/Internal/NNECoreOperator.h
@@ -19,7 +19,7 @@ namespace UE::NNECore::Internal
 		Atanh,
 		//BitShift, //Note: need integer tensors
 		Ceil,
-		//Clip,     //Note: need scalar tensor inputs
+		Clip,
 		Cos,
 		Cosh,
 		Elu,
diff --git a/Engine/Plugins/Experimental/NNE/Source/NNEHlslShaders/Internal/NNEHlslShadersBatchNormalizationCS.h b/Engine/Plugins/Experimental/NNE/Source/NNEHlslShaders/Internal/NNEHlslShadersBatchNormalizationCS.h
new file mode 100644
--- /dev/null
+++ b/Engine/Plugins/Experimental/NNE/Source/NNEHlslShaders/Internal/NNEHlslShadersBatchNormalizationCS.h
@@ -0,0 +1,42 @@
+// Copyright Epic Games, Inc. All Rights Reserved.
+
+#pragma once
+
+#include "GlobalShader.h"
+#include "ShaderParameterUtils.h"
+#include "RenderGraphUtils.h"
+#include "NNECoreOperator.h"
+
+namespace UE::NNEHlslShaders::Internal
+{
+	class FBatchNormalizationConstants
+	{
+	public:
+
+		static const int32 NUM_GROUP_THREADS{ 256 };
+	};
+
+	class NNEHLSLSHADERS_API TBatchNormalizationCS : public FGlobalShader
+	{
+		DECLARE_GLOBAL_SHADER(TBatchNormalizationCS);
+		SHADER_USE_PARAMETER_STRUCT(TBatchNormalizationCS, FGlobalShader)
+
+	public:
+
+		BEGIN_SHADER_PARAMETER_STRUCT(FParameters, )
+			SHADER_PARAMETER_RDG_BUFFER_SRV(Buffer<float>, X)
+			SHADER_PARAMETER_RDG_BUFFER_SRV(Buffer<float>, Scales)
+			SHADER_PARAMETER_RDG_BUFFER_SRV(Buffer<float>, Bias)
+			SHADER_PARAMETER_RDG_BUFFER_SRV(Buffer<float>, Mean)
+			SHADER_PARAMETER_RDG_BUFFER_SRV(Buffer<float>, Var)
+			SHADER_PARAMETER_RDG_BUFFER_UAV(RWBuffer<float>, Output)
+			SHADER_PARAMETER(uint32, Num)
+			SHADER_PARAMETER(uint32, DimC)
+			SHADER_PARAMETER(uint32, SpatialVolume)
+			SHADER_PARAMETER(uint32, ThreadCountX)
+			SHADER_PARAMETER(float, Epsilon)
+		END_SHADER_PARAMETER_STRUCT()
+
+		static void ModifyCompilationEnvironment(const FGlobalShaderPermutationParameters& InParameters, FShaderCompilerEnvironment& OutEnvironment);
+	};
+} // UE::NNEHlslShaders::Internal
diff --git a/Engine/Plugins/Experimental/NNE/Source/NNEHlslShaders/Private/NNEHlslShadersBatchNormalizationCS.cpp b/Engine/Plugins/Experimental/NNE/Source/NNEHlslShaders/Private/NNEHlslShadersBatchNormalizationCS.cpp
new file mode 100644
--- /dev/null
+++ b/Engine/Plugins/Experimental/NNE/Source/NNEHlslShaders/Private/NNEHlslShadersBatchNormalizationCS.cpp
@@ -0,0 +1,16 @@
+// Copyright Epic Games, Inc. All Rights Reserved.
+
+#include "NNEHlslShadersBatchNormalizationCS.h"
+#include "NNECore.h"
+
+namespace UE::NNEHlslShaders::Internal
+{
+	void TBatchNormalizationCS::ModifyCompilationEnvironment(const FGlobalShaderPermutationParameters& InParameters, FShaderCompilerEnvironment& OutEnvironment)
+	{
+		FGlobalShader::ModifyCompilationEnvironment(InParameters, OutEnvironment);
+
+		OutEnvironment.SetDefine(TEXT("THREADGROUP_SIZE_X"), FBatchNormalizationConstants::NUM_GROUP_THREADS);
+	}
+
+	IMPLEMENT_GLOBAL_SHADER(TBatchNormalizationCS, "/NNE/NNEHlslShadersBatchNormalization.usf", "BatchNormalization", SF_Compute);
+} // UE::NNEHlslShaders::Internal
diff --git a/Engine/Plugins/Experimental/NNE/Source/NNEHlslShaders/Private/NNEHlslShadersElementWiseUnaryCS.cpp b/Engine/Plugins/Experimental/NNE/Source/NNEHlslShaders/Private/NNEHlslShadersElementWiseUnaryCS.cpp
--- a/Engine/Plugins/Experimental/NNE/Source/NNEHlslShaders/Private/NNEHlslShadersElementWiseUnaryCS.cpp
+++ b/Engine/Plugins/Experimental/NNE/Source/NNEHlslShaders/Private/NNEHlslShadersElementWiseUnaryCS.cpp
@@ -38,7 +38,7 @@ namespace UE::NNEHlslShaders::Internal
 		//OP(BitShift,  TEXT("bitshift(X)"));
 		//OP(Cast,      TEXT("cast(X)"));
 		OP(Ceil,        TEXT("ceil(X)"));
-		//OP(Clip,      TEXT("clip(X)"));
+		OP(Clip,        TEXT("clipOp(X)"));
 		OP(Cos,         TEXT("cos(X)"));
 		OP(Cosh,        TEXT("cosh(X)"));
 		OP(Elu,         TEXT("elu(X)"));
diff --git a/Engine/Plugins/Experimental/NNE/Source/NNEHlslShaders/Shaders/NNEHlslShadersBatchNormalization.usf b/Engine/Plugins/Experimental/NNE/Source/NNEHlslShaders/Shaders/NNEHlslShadersBatchNormalization.usf
new file mode 100644
--- /dev/null
+++ b/Engine/Plugins/Experimental/NNE/Source/NNEHlslShaders/Shaders/NNEHlslShadersBatchNormalization.usf
@@ -0,0 +1,41 @@
+// Copyright Epic Games, Inc. All Rights Reserved.
+
+#include "/Engine/Public/Platform.ush"
+
+#define WORK_TYPE float
+
+Buffer<WORK_TYPE> X;
+Buffer<WORK_TYPE> Scales;
+Buffer<WORK_TYPE> Bias;
+Buffer<WORK_TYPE> Mean;
+Buffer<WORK_TYPE> Var;
+RWBuffer<WORK_TYPE> Output;
+float Epsilon;
+uint SpatialVolume;
+uint DimC;
+uint ThreadCountX;
+uint Num;
+
+WORK_TYPE ApplyBatchNormalization(WORK_TYPE x, WORK_TYPE input_mean, WORK_TYPE input_var, WORK_TYPE scale, WORK_TYPE bias)
+{
+	//[REDACTED_PROFILE_URL]/onnx/blob/main/docs/Operators.md#BatchNormalization
+	return (x - input_mean) / sqrt(input_var + Epsilon) * scale + bias;
+
+}
+
+
+[numthreads(THREADGROUP_SIZE_X, 1, 1)]
+void BatchNormalization(in const uint3 DispatchThreadID : SV_DispatchThreadID)
+{
+	const uint IndexX = DispatchThreadID.y * ThreadCountX + DispatchThreadID.x;
+	const uint IndexOther = (IndexX / SpatialVolume) % DimC;
+	if (IndexX < Num)
+	{
+		float x = X[IndexX];
+		float mean = Mean[IndexOther];
+		float var = Var[IndexOther];
+		float scale = Scales[IndexOther];
+		float bias = Bias[IndexOther];
+		Output[IndexX] = ApplyBatchNormalization(x, mean, var, scale, bias);
+	}
+}
diff --git a/Engine/Plugins/Experimental/NNE/Source/NNEHlslShaders/Shaders/NNEHlslShadersElementWiseUnary.usf b/Engine/Plugins/Experimental/NNE/Source/NNEHlslShaders/Shaders/NNEHlslShadersElementWiseUnary.usf
--- a/Engine/Plugins/Experimental/NNE/Source/NNEHlslShaders/Shaders/NNEHlslShadersElementWiseUnary.usf
+++ b/Engine/Plugins/Experimental/NNE/Source/NNEHlslShaders/Shaders/NNEHlslShadersElementWiseUnary.usf
@@ -107,6 +107,15 @@ float leakyRelu(float x)
 }
 
 
+//see [REDACTED_PROFILE_URL]/onnx/blob/main/docs/Operators.md#clip
+float clipOp(float x)
+{
+	float min = Alpha;
+	float max = Beta;
+	return clamp(x, Alpha, Beta);
+}
+
+
 [numthreads(THREADGROUP_SIZE_X, 1, 1)]
 void ElementWiseUnary(in const uint3 DispatchThreadID : SV_DispatchThreadID)
 {
diff --git a/Engine/Plugins/Experimental/NNE/Source/NNEQA/Private/NNEQAUnitTestOpUnary.cpp b/Engine/Plugins/Experimental/NNE/Source/NNEQA/Private/NNEQAUnitTestOpUnary.cpp
--- a/Engine/Plugins/Experimental/NNE/Source/NNEQA/Private/NNEQAUnitTestOpUnary.cpp
+++ b/Engine/Plugins/Experimental/NNE/Source/NNEQA/Private/NNEQAUnitTestOpUnary.cpp
@@ -118,6 +118,20 @@ namespace UE::NNEQA::Private::NNERuntimeRDG::UnaryOp
 		UTEST_EQUAL(TEXT("Ceil(XC2x1)[1]"), Y.GetPreparedData<float>()[1], FMath::CeilToFloat(XC2x1.GetPreparedData<float>()[1]));
 		return true;
 	}
+
+	IMPLEMENT_NNE_SHAPEINFERENCEHELPER_UNIT_AUTOMATION_TEST(FElementWiseUnaryCPUHelperClip, "System.Engine.MachineLearning.NNE.UnitTest.UnaryHelper.Clip");
+	bool FElementWiseUnaryCPUHelperClip::RunTest(const FString& Parameter)
+	{
+		FTensor XC2x1 = MakeConstTensor(TEXT("XC3x1"), { 3,1 }, { 0.5f, 1.8f, 2.1f });
+		FTensor Y = MakeTensor(TEXT("Y"), { 3,1 });
+		CPUHelper::ElementWiseUnary::Apply(EElementWiseUnaryOperatorType::Clip, XC2x1, 1.0f, 2.0f, 3.0f, Y);
+		UTEST_TRUE(TEXT("Y const if input is const"), Y.HasPreparedData());
+		UTEST_EQUAL(TEXT("Clip(XC3x1)[0]"), Y.GetPreparedData<float>()[0], 1.0f);
+		UTEST_EQUAL(TEXT("Clip(XC3x1)[1]"), Y.GetPreparedData<float>()[1], XC2x1.GetPreparedData<float>()[1]);
+		UTEST_EQUAL(TEXT("Clip(XC3x1)[2]"), Y.GetPreparedData<float>()[2], 2.0f);
+		return true;
+	}
+
 	IMPLEMENT_NNE_SHAPEINFERENCEHELPER_UNIT_AUTOMATION_TEST(FElementWiseUnaryCPUHelperCos, "System.Engine.MachineLearning.NNE.UnitTest.UnaryHelper.Cos");
 	bool FElementWiseUnaryCPUHelperCos::RunTest(const FString& Parameter)
 	{
diff --git a/Engine/Plugins/Experimental/NNE/Source/NNEQA/Resources/NNEQATestDesc.json b/Engine/Plugins/Experimental/NNE/Source/NNEQA/Resources/NNEQATestDesc.json
--- a/Engine/Plugins/Experimental/NNE/Source/NNEQA/Resources/NNEQATestDesc.json
+++ b/Engine/Plugins/Experimental/NNE/Source/NNEQA/Resources/NNEQATestDesc.json
@@ -13,8 +13,8 @@
 					"absolute_tolerance": 5.0e-2
 				}
 			]
-		},
-		{
+				},
+				{
 			"category": "StyleTransfer.VariableInputSize",
 			"is_model_test": true,
 			"targets": [
@@ -156,13 +156,12 @@
 				{	"target": "BitShift", "skip": true, "comment": "TODO need attributes" },
 				{	"target": "Cast",     "skip": true, "comment": "TODO need attributes" },
 				{	"target": "Ceil"},
-				{	"target": "Clip", "additional_datasets": [ "Elementwise.ClipTestWithExtraInput" ], "runtimes": [
-						{ "name": "NNERuntimeRDGDml", "skip": true, "comment": "RDG need scalar tensor support"},
-						{ "name": "NNERuntimeRDGHlsl", "skip": true, "comment": "RDG need scalar tensor support"}						
+				{	"target": "Clip", "additional_datasets": [ "Elementwise.Clip.AttributeTest" ], "runtimes": [
+						{ "name": "NNERuntimeRDGDml", "skip": true, "comment": "RDG need scalar tensor support"}
 					]},
 				{	"target": "Cos" },
 				{	"target": "Cosh" },
-				{	"target": "Elu", "comment": "TODO test with not-default attribute", "runtimes": [
+				{	"target": "Elu", "additional_datasets": [ "Elementwise.Elu.AttributeTest" ], "runtimes": [
 						{ "name": "NNERuntimeORTDml", "absolute_tolerance":  1e-4},
 						{ "name": "NNERuntimeRDGDml",    "absolute_tolerance":  1e-4},
 						{ "name": "NNERuntimeRDGHlsl",   "absolute_tolerance":  1e-4}
@@ -176,7 +175,7 @@
 				{	"target": "Floor" },
 				{	"target": "IsInf", "skip": true, "comment": "TODO add flexibility in test data setup + bool tensors" },
 				{	"target": "IsNan", "skip": true, "comment": "TODO add flexibility in test data setup + bool tensors" },
-				{	"target": "HardSigmoid",         "comment": "TODO test with not-default attribute" },
+				{	"target": "HardSigmoid", "additional_datasets": [ "Elementwise.HardSigmoid.AttributeTest" ] },
 				{	"target": "HardSwish", "runtimes": [ { "name": "NNERuntimeRDGDml", "skip": true, "comment": "Not implemented atm on NNERuntimeRDGDml."} ]},
 				{	"target": "LeakyRelu", "additional_datasets": [ "Elementwise.LeakyRelu.AttributeTest" ] },
 				{	"target": "Log", "runtimes": [
@@ -192,7 +191,7 @@
 				{	"target": "Reciprocal"},
 				{	"target": "Relu"},
 				{	"target": "Round"},
-				{	"target": "Selu", "comment": "TODO test with not-default attribute", "runtimes": [
+				{	"target": "Selu", "additional_datasets": [ "Elementwise.Selu.AttributeTest" ], "runtimes": [
 						{ "name": "NNERuntimeORTDml", "absolute_tolerance":  1e-4},
 						{ "name": "NNERuntimeRDGDml",    "absolute_tolerance":  1e-4},
 						{ "name": "NNERuntimeRDGHlsl",   "absolute_tolerance":  1e-4}
@@ -309,6 +308,42 @@
 				{ "target": "Pad" }
 			]
 		},
+		{
+			"category": "Identity",
+			"runtimes": [
+				{ "name": "NNERuntimeRDGDml", "skip": true },
+			],
+			"targets": [
+				{ "target": "Identity" }
+			]
+		},
+		{
+			"category": "Dropout",
+			"runtimes": [
+				{ "name": "NNERuntimeRDGDml", "skip": true },
+			],
+			"targets": [
+				{ "target": "Dropout" }
+			]
+		},
+		{
+			"category": "Flatten",
+			"runtimes": [
+				{ "name": "NNERuntimeRDGDml", "skip": true },
+			],
+			"targets": [
+				{ "target": "Flatten" }
+			]
+		},
+		{
+			"category": "BatchNormalization",
+			"runtimes": [
+				{ "name": "NNERuntimeRDGDml", "skip": true },
+			],
+			"targets": [
+				{ "target": "BatchNormalization", "relative_tolerance": 5e-1 }
+			]
+		},
 		{
 			"category": "InstanceNormalization",
 			"targets": [
@@ -379,8 +414,6 @@
 		},
 		{
 			"category": "Reshape",
-			"runtimes": [
-			],
 			"targets": [
 				{ "target": "Reshape" }
 			]
@@ -454,6 +487,38 @@
 				}
 			]
 		},
+		{
+			"name": "Elementwise.Elu.AttributeTest",
+			"datasets": [
+				{
+					"inputs": [ { "shape": [ 4 ] } ]
+				}
+			]
+		},
+		{
+			"name": "Elementwise.Selu.AttributeTest",
+			"datasets": [
+				{
+					"inputs": [ { "shape": [ 4 ] } ]
+				}
+			]
+		},
+		{
+			"name": "Elementwise.HardSigmoid.AttributeTest",
+			"datasets": [
+				{
+					"inputs": [ { "shape": [ 4 ] } ]
+				}
+			]
+		},
+		{
+			"name": "Elementwise.Clip.AttributeTest",
+			"datasets": [
+				{
+					"inputs": [ { "shape": [ 4 ] } ]
+				}
+			]
+		},
 		{
 			"name": "Elementwise.LeakyRelu.AttributeTest",
 			"datasets": [
@@ -734,6 +799,55 @@
 				{ "inputs":  [ { "shape": [ 2, 1, 3, 1, 2 ] } ],  "weights": [ { "shape": [ 5 ], "source": [ "1", "2", "2", "1", "2" ], "type": "float" } ], "outputs": [ { "shape": [ 2, 2, 6, 1, 4 ]} ] }
 			]
 		},
+		{
+			"name": "Identity",
+			"datasets": [
+				{ 
+					"inputs": [ { "shape": [ 1, 2, 3, 4 ] } ], "outputs": [ { "shape": [ 1, 2, 3, 4 ] } ]
+				}
+			]
+		},
+		{
+			"name": "Dropout",
+			"datasets": [
+				{ 
+					"inputs": [ { "shape": [ 1, 2, 3, 4 ] } ], "outputs": [ { "shape": [ 1, 2, 3, 4 ] } ]
+				}
+			]
+		},
+		{ 
+			"name": "Flatten.Axis2",
+			"datasets": [
+				{ "inputs": [ { "shape": [ 2, 3, 1, 4 ] } ], "outputs": [ { "shape": [ 6, 4 ] } ] }
+			]
+		},
+		{ 
+			"name": "Flatten.Axis1",
+			"datasets": [
+				{ "inputs": [ { "shape": [ 2, 3, 1, 4 ] } ], "outputs": [ { "shape": [ 2, 12 ] } ] }
+			]
+		},
+		{ 
+			"name": "Flatten.Axis0",
+			"datasets": [
+				{ "inputs": [ { "shape": [ 2, 3, 1, 4 ] } ], "outputs": [ { "shape": [ 1, 24 ] } ] }
+			]
+		},
+		{
+			"name": "BatchNormalization",
+			"datasets": [
+				{ 
+					"inputs": [ { "shape": [ 2, 3, 4 ] }, { "shape": [ 3 ] }, { "shape": [ 3 ] }, { "shape": [ 3 ] },
+								{ "shape": [ 3 ], "source": [ "0.1", "1.0", "2.0" ], "comments":  "variance should be positive"}],
+					"outputs": [ { "shape": [ 2, 3, 4 ] } ]
+				},
+				{ 
+					"inputs": [ { "shape": [ 1, 8, 32, 64 ] }, { "shape": [ 8 ] }, { "shape": [ 8 ] }, { "shape": [ 8 ] },
+								{ "shape": [ 8 ], "source": [ "0.0", "0.1", "1.2", "3.0", "0.4", "2.0", "6.0", "0.5" ], "comments":  "variance should be positive"}],
+					"outputs": [ { "shape": [ 1, 8, 32, 64 ] } ]
+				}
+			]
+		},
 		{
 			"name": "InstanceNormalization",
 			"datasets": [
@@ -912,20 +1026,74 @@
 		{
 			"name": "Reshape.MinusOne",
 			"datasets": [
-				{ "inputs": [ { "shape": [ 2, 4, 3 ] } ], "weights": [ { "shape": [ 2 ], "source": [ "12", "-1" ], "type": "int64" } ], "outputs": [ { "shape": [ 12, 2 ] } ], 
-				  "runtimes": [ { "name": "NNERuntimeRDGHlsl", "skip": true, "comment": "Shape inference here not yet supported" } ] }
+				{ "inputs": [ { "shape": [ 2, 4, 3 ] } ], "weights": [ { "shape": [ 2 ], "source": [ "12", "-1" ], "type": "int64" } ], "outputs": [ { "shape": [ 12, 2 ] } ] }
 			]
 		},
 		{
 			"name": "Reshape.Zero",
 			"datasets": [
-				{ "inputs": [ { "shape": [ 2, 4, 3 ] } ], "weights": [ { "shape": [ 3 ], "source": [ "3", "0", "2" ], "type": "int64" } ], "outputs": [ { "shape": [ 3, 4, 2 ] } ], 
-				  "runtimes": [ { "name": "NNERuntimeRDGHlsl", "skip": true, "comment": "Shape inference here not yet supported" } ] }
+				{ "inputs": [ { "shape": [ 2, 4, 3 ] } ], "weights": [ { "shape": [ 3 ], "source": [ "3", "0", "2" ], "type": "int64" } ], "outputs": [ { "shape": [ 3, 4, 2 ] } ] }
 			]
 		}
 	],
 
 	"attribute_sets": [
+		{
+			"name": "Flatten.Axis0",
+			"attribute_maps": [
+				{ "attributes": [ { "name": "axis", "value": 0, "type": "Int32" } ] }
+			]
+		},
+		{
+			"name": "Flatten.Axis1",
+			"attribute_maps": [
+				{ "attributes": [ { "name": "axis", "value": 1, "type": "Int32" } ] }
+			]
+		},
+		{
+			"name": "Flatten.Axis2",
+			"attribute_maps": [
+				{ "attributes": [ { "name": "axis", "value": 2, "type": "Int32" } ] }
+			]
+		},
+		{
+			"name": "Elementwise.Elu.AttributeTest",
+			"attribute_maps": [
+				{ "attributes": [ { "name": "alpha", "value": 0.5 } ] },
+				{ "attributes": [ { "name": "alpha", "value": 1.0 } ] },
+				{ "attributes": [ ] }
+			]
+		},
+		{
+			"name": "Elementwise.Selu.AttributeTest",
+			"attribute_maps": [
+				{ "attributes": [ { "name": "alpha", "value": 0.5 } ] },
+				{ "attributes": [ { "name": "alpha", "value": 1.0 } ] },
+				{ "attributes": [ { "name": "alpha", "value": 0.7 }, { "name": "gamma", "value": 0.1 } ] },
+				{ "attributes": [ { "name": "alpha", "value": 0.7 }, { "name": "gamma", "value": 0.5 } ] },
+				{ "attributes": [ ] }
+			]
+		},
+		{
+			"name": "Elementwise.HardSigmoid.AttributeTest",
+			"attribute_maps": [
+				{ "attributes": [ { "name": "alpha", "value": 0.5 } ] },
+				{ "attributes": [ { "name": "alpha", "value": 1.0 } ] },
+				{ "attributes": [ { "name": "alpha", "value": 0.7 }, { "name": "beta", "value": 0.1 } ] },
+				{ "attributes": [ { "name": "alpha", "value": 0.7 }, { "name": "beta", "value": 0.5 } ] },
+				{ "attributes": [ ] }
+			]
+		},
+		{
+			"name": "Elementwise.Clip.AttributeTest",
+			"attribute_maps": [
+				{ "attributes": [ { "name": "min", "value": 0.1 } ] },
+				{ "attributes": [ { "name": "max", "value": 0.5 } ] },
+				{ "attributes": [ { "name": "min", "value": -0.1 }, { "name": "max", "value": 0.1 } ] },
+				{ "attributes": [ { "name": "min", "value": -0.1 }, { "name": "max", "value": 0.5 } ] },
+				{ "attributes": [ ] }
+			]
+		},
 		{
 			"name": "Elementwise.LeakyRelu.AttributeTest",
 			"attribute_maps": [
diff --git a/Engine/Plugins/Experimental/NNE/Source/NNERuntimeRDG/Private/Dml/NNEDmlOperatorReshape.cpp b/Engine/Plugins/Experimental/NNE/Source/NNERuntimeRDG/Private/Dml/NNEDmlOperatorReshape.cpp
--- a/Engine/Plugins/Experimental/NNE/Source/NNERuntimeRDG/Private/Dml/NNEDmlOperatorReshape.cpp
+++ b/Engine/Plugins/Experimental/NNE/Source/NNERuntimeRDG/Private/Dml/NNEDmlOperatorReshape.cpp
@@ -6,6 +6,7 @@
 #include "Algo/Transform.h"
 #include "Algo/Count.h"
 #include "Misc/EnumerateRange.h"
+#include "Helper/NNERuntimeRDGHelperReshape.h"
 
 namespace UE::NNERuntimeRDG::Private::Dml
 {
@@ -16,72 +17,6 @@ namespace UE::NNERuntimeRDG::Private::Dml
 class FOperatorDmlReshape : public FOperatorDml
 {
 
-    template<typename DataType>
-    bool ReshapeTensorToDmlShape(
-        const NNECore::Internal::FTensor& InputTensor, 
-        const NNECore::Internal::FTensor& ShapeTensor, 
-        bool bAllowZero,
-        DmlUtil::FSmallUIntArray& OutShape)
-    {
-        
-        DmlUtil::FSmallArray<DataType> ReshapedShape(ShapeTensor.GetPreparedData<DataType>());
-
-        if(!bAllowZero)
-        {
-            // at most 1 dimension can be -1
-            if(Algo::Count(ReshapedShape, -1) > 1)
-            {
-                UE_LOG(LogNNE, Error, TEXT("Shape tensor can't contain more than one '-1'."));
-                return false;
-            }
-            for (TEnumerateRef<DataType> Elem : EnumerateRange(ReshapedShape))
-            {
-                if(!(*Elem != 0 || InputTensor.GetShape().Rank() > Elem.GetIndex()))
-                {
-                    UE_LOG(LogNNE, Error, TEXT("Shape tensor contains '0' in an invalid place."));
-                    return false;
-                }
-                *Elem = 
-                    *Elem == 0 ? 
-                        InputTensor.GetShape().GetData()[Elem.GetIndex()] 
-                        : 
-                        *Elem;
-            }
-        }
-        else
-        {
-            // no -1 is allowed if there is a 0
-            if(Algo::Count(ReshapedShape, 0) != 0 && Algo::Count(ReshapedShape, -1) != 0)
-            {
-                UE_LOG(LogNNE, Error, TEXT("Shape tensor contains both '0' and '-1'. This is not allowed."));
-                return false;
-            }
-        }
-
-        auto PartialVolume = [](TConstArrayView<DataType> Shape) -> uint32
-            {
-                uint32 Volume = 1;
-                for(DataType Dim : Shape)
-                {
-                    Volume *= Dim != -1 ? (DataType) Dim : 1;
-                }
-                return Volume;
-            };
-
-        OutShape.SetNumUninitialized(ReshapedShape.Num());
-
-        for(TConstEnumerateRef<DataType> Elem : EnumerateRange(ReshapedShape))
-        {
-            OutShape[Elem.GetIndex()] = 
-                *Elem == -1 ? 
-                    (uint32) InputTensor.GetShape().Volume() / PartialVolume(ReshapedShape)
-                    :
-                    (uint32) *Elem;
-        }
-
-        return true;
-    }
-
 public:
 
 	static FOperatorDml* Create()
@@ -110,19 +45,19 @@ public:
         switch(InputTensors[1].GetDataType())
         {
         case ENNETensorDataType::Int32:
-            if(!ReshapeTensorToDmlShape<int32>(InputTensors[0], InputTensors[1], bAllowZero, ReshapedShape))
+            if(!ShapeHelper::Reshape::ReshapeTensor<int32>(InputTensors[0].GetShape(), InputTensors[1], bAllowZero, ReshapedShape))
             {
                 return false;
             }
             break;
         case ENNETensorDataType::Int64:
-            if(!ReshapeTensorToDmlShape<int64>(InputTensors[0], InputTensors[1], bAllowZero, ReshapedShape))
+			if (!ShapeHelper::Reshape::ReshapeTensor<int64>(InputTensors[0].GetShape(), InputTensors[1], bAllowZero, ReshapedShape))
             {
                 return false;
             }
             break;
         case ENNETensorDataType::UInt32:
-            if(!ReshapeTensorToDmlShape<uint32>(InputTensors[0], InputTensors[1], bAllowZero, ReshapedShape))
+			if (!ShapeHelper::Reshape::ReshapeTensor<uint32>(InputTensors[0].GetShape(), InputTensors[1], bAllowZero, ReshapedShape))
             {
                 return false;
             }
diff --git a/Engine/Plugins/Experimental/NNE/Source/NNERuntimeRDG/Private/Helper/NNERuntimeRDGHelperElementWiseUnary.cpp b/Engine/Plugins/Experimental/NNE/Source/NNERuntimeRDG/Private/Helper/NNERuntimeRDGHelperElementWiseUnary.cpp
--- a/Engine/Plugins/Experimental/NNE/Source/NNERuntimeRDG/Private/Helper/NNERuntimeRDGHelperElementWiseUnary.cpp
+++ b/Engine/Plugins/Experimental/NNE/Source/NNERuntimeRDG/Private/Helper/NNERuntimeRDGHelperElementWiseUnary.cpp
@@ -41,6 +41,8 @@ namespace UE::NNERuntimeRDG::Internal::CPUHelper::ElementWiseUnary
 	}
 
 	template<> float Apply<NNECore::Internal::EElementWiseUnaryOperatorType::Ceil>(float X, float Alpha, float Beta, float Gamma) { return FMath::CeilToFloat(X); }
+
+	template<> float Apply<NNECore::Internal::EElementWiseUnaryOperatorType::Clip>(float X, float Alpha, float Beta, float Gamma) { return FMath::Clamp(X, Alpha, Beta); }
 	
 	template<> float Apply<NNECore::Internal::EElementWiseUnaryOperatorType::Cos>(float X, float Alpha, float Beta, float Gamma) { return FMath::Cos(X); }
 	
@@ -170,6 +172,9 @@ namespace UE::NNERuntimeRDG::Internal::CPUHelper::ElementWiseUnary
 		case NNECore::Internal::EElementWiseUnaryOperatorType::Ceil:
 			Apply<NNECore::Internal::EElementWiseUnaryOperatorType::Ceil>(Tensor, Alpha, Beta, Gamma, OutputTensor);
 			break;
+		case NNECore::Internal::EElementWiseUnaryOperatorType::Clip:
+			Apply<NNECore::Internal::EElementWiseUnaryOperatorType::Clip>(Tensor, Alpha, Beta, Gamma, OutputTensor);
+			break;
 		case NNECore::Internal::EElementWiseUnaryOperatorType::Cos:
 			Apply<NNECore::Internal::EElementWiseUnaryOperatorType::Cos>(Tensor, Alpha, Beta, Gamma, OutputTensor);
 			break;
diff --git a/Engine/Plugins/Experimental/NNE/Source/NNERuntimeRDG/Private/Helper/NNERuntimeRDGHelperReshape.h b/Engine/Plugins/Experimental/NNE/Source/NNERuntimeRDG/Private/Helper/NNERuntimeRDGHelperReshape.h
new file mode 100644
--- /dev/null
+++ b/Engine/Plugins/Experimental/NNE/Source/NNERuntimeRDG/Private/Helper/NNERuntimeRDGHelperReshape.h
@@ -0,0 +1,80 @@
+// Copyright Epic Games, Inc. All Rights Reserved.
+
+#pragma once 
+
+#include "Algo/Transform.h"
+#include "Algo/Count.h"
+#include "CoreMinimal.h"
+#include "Misc/EnumerateRange.h"
+#include "NNECoreTensor.h"
+#include "NNECoreTypes.h"
+
+namespace UE::NNERuntimeRDG::Private::ShapeHelper::Reshape
+{
+	template<typename DataType, class OutputShapeType >
+	bool ReshapeTensor(
+		const NNECore::FTensorShape& InputTensorShape,
+		const NNECore::Internal::FTensor& ShapeTensor,
+		bool bAllowZero,
+		OutputShapeType& OutShape)
+	{
+
+		TArray<DataType, TInlineAllocator<NNECore::FTensorShape::MaxRank>> ReshapedShape(ShapeTensor.GetPreparedData<DataType>());
+
+		if (!bAllowZero)
+		{
+			// at most 1 dimension can be -1
+			if (Algo::Count(ReshapedShape, -1) > 1)
+			{
+				UE_LOG(LogNNE, Error, TEXT("Reshape: Shape tensor can't contain more than one '-1'."));
+				return false;
+			}
+			for (TEnumerateRef<DataType> Elem : EnumerateRange(ReshapedShape))
+			{
+				if (!(*Elem != 0 || InputTensorShape.Rank() > Elem.GetIndex()))
+				{
+					UE_LOG(LogNNE, Error, TEXT("Reshape: Shape tensor contains '0' in an invalid place."));
+					return false;
+				}
+				*Elem =
+					*Elem == 0 ?
+					InputTensorShape.GetData()[Elem.GetIndex()]
+					:
+					*Elem;
+			}
+		}
+		else
+		{
+			// no -1 is allowed if there is a 0
+			if (Algo::Count(ReshapedShape, 0) != 0 && Algo::Count(ReshapedShape, -1) != 0)
+			{
+				UE_LOG(LogNNE, Error, TEXT("Reshape: Shape tensor contains both '0' and '-1'. This is not allowed."));
+				return false;
+			}
+		}
+
+		auto PartialVolume = [](TConstArrayView<DataType> Shape) -> uint32
+		{
+			uint32 Volume = 1;
+			for (DataType Dim : Shape)
+			{
+				Volume *= Dim != -1 ? (DataType)Dim : 1;
+			}
+			return Volume;
+		};
+
+		OutShape.SetNumUninitialized(ReshapedShape.Num());
+
+		for (TConstEnumerateRef<DataType> Elem : EnumerateRange(ReshapedShape))
+		{
+			OutShape[Elem.GetIndex()] =
+				*Elem == -1 ?
+				(uint32)InputTensorShape.Volume() / PartialVolume(ReshapedShape)
+				:
+				(uint32)*Elem;
+		}
+
+		return true;
+	}
+
+} // UE::NNERuntimeRDG::Private::ShapeHelper::Reshape
diff --git a/Engine/Plugins/Experimental/NNE/Source/NNERuntimeRDG/Private/Hlsl/NNERuntimeRDGBatchNormalization.cpp b/Engine/Plugins/Experimental/NNE/Source/NNERuntimeRDG/Private/Hlsl/NNERuntimeRDGBatchNormalization.cpp
new file mode 100644
--- /dev/null
+++ b/Engine/Plugins/Experimental/NNE/Source/NNERuntimeRDG/Private/Hlsl/NNERuntimeRDGBatchNormalization.cpp
@@ -0,0 +1,160 @@
+// Copyright Epic Games, Inc. All Rights Reserved.
+
+#include "NNERuntimeRDGBatchNormalization.h"
+#include "NNEHlslShadersBatchNormalizationCS.h"
+#include "NNERuntimeRDGHlslHelper.h"
+#include "NNECoreAttributeMap.h"
+#include "NNECoreTypes.h"
+#include "NNECoreTensor.h"
+
+namespace UE::NNERuntimeRDG::Private::Hlsl
+{
+	DECLARE_GPU_STAT_NAMED(FNNEOperatorBatchNormalization, TEXT("NNE.Operator.Hlsl.BatchNormalization"));
+
+	/**
+	 * BatchNormalization operator implementation
+	 */
+	class FBatchNormalization : public FOperatorHlsl
+	{
+	public:
+
+		FBatchNormalization() {}
+		virtual ~FBatchNormalization() = default;
+
+		float Epsilon;
+
+	public:
+
+		virtual int PrepareOutputs(TConstArrayView<NNECore::Internal::FTensorRef> InputTensors, TArrayView<NNECore::Internal::FTensorRef> OutputTensors) const override
+		{
+			check(InputTensors.Num() == 5);
+			check(OutputTensors.Num() == 1);
+
+			const NNECore::Internal::FTensor& X = *InputTensors[0];
+
+			OutputTensors[0]->SetShape(X.GetShape());
+
+			return 0;
+		};
+
+		virtual bool Initialize(TConstArrayView<NNECore::FTensorDesc> InputTensorDescs, TConstArrayView<NNECore::FTensorDesc> OutputTensorDescs, const NNECore::FAttributeMap& Attributes) override
+		{
+			check(InputTensorDescs.Num() == 5);
+			check(OutputTensorDescs.Num() >= 1);
+
+			Epsilon = Attributes.GetValueOrDefault<float>(TEXT("epsilon"), 1e-5f);
+
+			if (OutputTensorDescs.Num() > 1)
+			{
+				UE_LOG(LogNNE, Warning, TEXT("BatchNormalization is only supported in inference mode at the moment, supporting no more than one output."));
+				return false;
+			}
+			
+			return true;
+		}
+
+		virtual void Dispatch(FRDGBuilder& GraphBuilder, TConstArrayView<FTensorRDGRef> InputTensors, TConstArrayView<FTensorRDGRef> OutputTensors) override
+		{
+			using namespace UE::NNEHlslShaders::Internal;
+
+			check(InputTensors.Num() == 5);
+			check(OutputTensors.Num() == 1);
+			check(InputTensors[0] != nullptr);
+			check(InputTensors[1] != nullptr);
+			check(InputTensors[2] != nullptr);
+			check(InputTensors[3] != nullptr);
+			check(InputTensors[4] != nullptr);
+			check(OutputTensors[0] != nullptr);
+
+			const FTensorRDG& X = *InputTensors[0];
+			const FTensorRDG& Scales = *InputTensors[1];
+			const FTensorRDG& Bias = *InputTensors[2];
+			const FTensorRDG& Mean = *InputTensors[3];
+			const FTensorRDG& Var = *InputTensors[4];
+			const FTensorRDG& Output = *OutputTensors[0];
+
+			FRDGBufferSRVRef XSRV = GraphBuilder.CreateSRV(FRDGBufferSRVDesc(X.GetBuffer(), PF_R32_FLOAT));
+			FRDGBufferSRVRef ScalesSRV = GraphBuilder.CreateSRV(FRDGBufferSRVDesc(Scales.GetBuffer(), PF_R32_FLOAT));
+			FRDGBufferSRVRef BiasSRV = GraphBuilder.CreateSRV(FRDGBufferSRVDesc(Bias.GetBuffer(), PF_R32_FLOAT));
+			FRDGBufferSRVRef MeanSRV = GraphBuilder.CreateSRV(FRDGBufferSRVDesc(Mean.GetBuffer(), PF_R32_FLOAT));
+			FRDGBufferSRVRef VarSRV = GraphBuilder.CreateSRV(FRDGBufferSRVDesc(Var.GetBuffer(), PF_R32_FLOAT));
+			FRDGBufferUAVRef OutputUAV = GraphBuilder.CreateUAV(FRDGBufferUAVDesc(Output.GetBuffer(), PF_R32_FLOAT));
+
+			FIntVector ThreadGroupCount = ComputeElementWiseThreadGroups(Output.GetVolume(), FBatchNormalizationConstants::NUM_GROUP_THREADS);
+			int32 DimC = 1;
+			int32 SpatialVolume = 1;
+			if (X.GetShape().Rank() >= 2)
+			{
+				DimC = X.GetShape().GetData()[1];
+			}
+			for (int32 i = X.GetShape().Rank() - 1; i > 1; --i)
+			{
+				SpatialVolume *= X.GetShape().GetData()[i];
+			}
+
+			// Set parameters
+			TBatchNormalizationCS::FParameters* Params = GraphBuilder.AllocParameters<TBatchNormalizationCS::FParameters>();
+			Params->X = XSRV;
+			Params->Scales = ScalesSRV;
+			Params->Bias = BiasSRV;
+			Params->Mean = MeanSRV;
+			Params->Var = VarSRV;
+			Params->Output = OutputUAV;
+			Params->Num = Output.GetVolume();
+			Params->DimC = DimC;
+			Params->SpatialVolume = SpatialVolume;
+			Params->ThreadCountX = ThreadGroupCount.X * FBatchNormalizationConstants::NUM_GROUP_THREADS;
+			Params->Epsilon = Epsilon;
+
+			TBatchNormalizationCS::FPermutationDomain PermutationVector;
+			TShaderMapRef<TBatchNormalizationCS> ComputeShader(GetGlobalShaderMap(GMaxRHIFeatureLevel), PermutationVector);
+
+			RDG_EVENT_SCOPE(GraphBuilder, "NNE.Operator.Hlsl.BatchNormalization");
+			RDG_GPU_STAT_SCOPE(GraphBuilder, FNNEOperatorBatchNormalization);
+
+			FComputeShaderUtils::AddPass(
+				GraphBuilder,
+				RDG_EVENT_NAME("NNE.Operator.Hlsl.BatchNormalization.Dispatch"),
+				ERDGPassFlags::Compute | ERDGPassFlags::NeverCull,
+				ComputeShader,
+				Params,
+				ThreadGroupCount);
+		}
+	};
+
+	bool ValidateBatchNormalizationOperator(const NNECore::FAttributeMap& AttributeMap, TConstArrayView<ENNETensorDataType> InputTypes, TConstArrayView<NNECore::FSymbolicTensorShape> InputBatchNormalizations)
+	{
+		bool bIsValid = true;
+
+		//This match version 9 of the BatchNormalization operator
+		//[REDACTED_PROFILE_URL]/onnx/blob/main/docs/Operators.md#BatchNormalization
+		FAttributeValidator AttributeValidator;
+		AttributeValidator.AddOptional(TEXT("epsilon"), ENNEAttributeDataType::Float);
+		AttributeValidator.AddOptional(TEXT("momentum"), ENNEAttributeDataType::Float); //Will be ignored, only useful in training mode.
+		bIsValid &= AttributeValidator.Validate(AttributeMap);
+
+		FInputValidator InputValidator;
+		InputValidator.AddSupportedType(ENNETensorDataType::Half);
+		InputValidator.AddSupportedType(ENNETensorDataType::Float);
+		InputValidator.AddSupportedType(ENNETensorDataType::Double);
+		InputValidator.AddRequired();
+		InputValidator.AddRequired();
+		InputValidator.AddRequired();
+		InputValidator.AddRequired();
+		InputValidator.AddRequired();
+		bIsValid &= InputValidator.Validate(InputTypes);
+
+		return bIsValid;
+	}
+
+	FOperatorHlsl* CreateBatchNormalizationOperator()
+	{
+		return new FBatchNormalization();
+	}
+
+	bool RegisterBatchNormalizationOperator(FOperatorRegistryHlsl& Registry)
+	{
+		Registry.OpAdd(TEXT("BatchNormalization"), CreateBatchNormalizationOperator, ValidateBatchNormalizationOperator);
+		return true;
+	}
+} // UE::NNERuntimeRDG::Private::Hlsl
diff --git a/Engine/Plugins/Experimental/NNE/Source/NNERuntimeRDG/Private/Hlsl/NNERuntimeRDGBatchNormalization.h b/Engine/Plugins/Experimental/NNE/Source/NNERuntimeRDG/Private/Hlsl/NNERuntimeRDGBatchNormalization.h
new file mode 100644
--- /dev/null
+++ b/Engine/Plugins/Experimental/NNE/Source/NNERuntimeRDG/Private/Hlsl/NNERuntimeRDGBatchNormalization.h
@@ -0,0 +1,11 @@
+// Copyright Epic Games, Inc. All Rights Reserved.
+
+#pragma once
+
+#include "NNERuntimeRDG.h"
+#include "NNERuntimeRDGHlslOp.h"
+
+namespace UE::NNERuntimeRDG::Private::Hlsl
+{
+	bool RegisterBatchNormalizationOperator(FOperatorRegistryHlsl& Registry);
+} // UE::NNERuntimeRDG::Private::Hlsl
diff --git a/Engine/Plugins/Experimental/NNE/Source/NNERuntimeRDG/Private/Hlsl/NNERuntimeRDGDropout.cpp b/Engine/Plugins/Experimental/NNE/Source/NNERuntimeRDG/Private/Hlsl/NNERuntimeRDGDropout.cpp
new file mode 100644
--- /dev/null
+++ b/Engine/Plugins/Experimental/NNE/Source/NNERuntimeRDG/Private/Hlsl/NNERuntimeRDGDropout.cpp
@@ -0,0 +1,107 @@
+// Copyright Epic Games, Inc. All Rights Reserved.
+
+#include "NNERuntimeRDGDropout.h"
+#include "NNECoreTensor.h"
+#include "NNECoreTypes.h"
+#include "RenderGraphUtils.h"
+
+namespace UE::NNERuntimeRDG::Private::Hlsl
+{
+	DECLARE_GPU_STAT_NAMED(FNNEOperatorDropout, TEXT("NNE.Operator.Hlsl.Dropout"));
+
+	/**
+	 * Dropout operator implementation
+	 */
+	class FDropout : public FOperatorHlsl
+	{
+	public:
+
+		FDropout() {}
+		virtual ~FDropout() = default;
+
+	public:
+
+		virtual int PrepareOutputs(TConstArrayView<NNECore::Internal::FTensorRef> InputTensors, TArrayView<NNECore::Internal::FTensorRef> OutputTensors) const override
+		{
+			check(InputTensors.Num() == 1);
+			check(OutputTensors.Num() == 1);
+
+			const NNECore::Internal::FTensor& X = *InputTensors[0];
+
+			OutputTensors[0]->SetShape(X.GetShape());
+			if (X.HasPreparedData())
+			{
+				OutputTensors[0]->SetPreparedData<uint8>(X.GetPreparedData<uint8>());
+			}
+
+			return 0;
+		};
+
+		virtual bool Initialize(TConstArrayView<NNECore::FTensorDesc> InputTensorDescs, TConstArrayView<NNECore::FTensorDesc> OutputTensorDescs, const NNECore::FAttributeMap& Attributes) override
+		{
+			check(InputTensorDescs.Num() == 1);
+			check(OutputTensorDescs.Num() == 1 || OutputTensorDescs.Num() == 2);
+
+			if (OutputTensorDescs.Num() == 2)
+			{
+				UE_LOG(LogNNE, Warning, TEXT("Dropout is only supported in inference mode at the moment, without the ability to output a 'mask'."));
+				return false;
+			}
+			
+			return true;
+		}
+
+		virtual void Dispatch(FRDGBuilder& GraphBuilder, TConstArrayView<FTensorRDGRef> InputTensors, TConstArrayView<FTensorRDGRef> OutputTensors) override
+		{
+			check(InputTensors.Num() == 1);
+			check(OutputTensors.Num() == 1 || OutputTensors.Num() == 2);
+			check(InputTensors[0] != nullptr);
+			check(OutputTensors[0] != nullptr);
+
+			const FTensorRDG& Data = *InputTensors[0];
+			const FTensorRDG& Output = *OutputTensors[0];
+
+			RDG_EVENT_SCOPE(GraphBuilder, "NNE.Operator.Hlsl.Dropout");
+			RDG_GPU_STAT_SCOPE(GraphBuilder, FNNEOperatorDropout);
+
+			AddCopyBufferPass(GraphBuilder, Output.GetBuffer(), Data.GetBuffer());
+		}
+	};
+
+	bool ValidateDropoutOperator(const NNECore::FAttributeMap& AttributeMap, TConstArrayView<ENNETensorDataType> InputTypes, TConstArrayView<NNECore::FSymbolicTensorShape> InputDropouts)
+	{
+		bool bIsValid = true;
+
+		//This match version 13 of the Dropout operator
+		//[REDACTED_PROFILE_URL]/onnx/blob/main/docs/Operators.md#Dropout
+		FAttributeValidator AttributeValidator;
+		AttributeValidator.AddOptional(TEXT("seed"), ENNEAttributeDataType::Int32); //Will be ignored, only useful in training mode.
+		bIsValid &= AttributeValidator.Validate(AttributeMap);
+
+		FInputValidator InputValidator;
+		InputValidator.AddSupportedType(ENNETensorDataType::Half);
+		InputValidator.AddSupportedType(ENNETensorDataType::Float);
+		InputValidator.AddSupportedType(ENNETensorDataType::Double);
+		InputValidator.AddRequired();
+
+		if (InputTypes.Num() > 1)
+		{
+			UE_LOG(LogNNE, Warning, TEXT("Dropout is only supported in inference mode at the moment, 'ratio' and 'training_mode' will be ignored."));
+		}
+
+		bIsValid &= InputValidator.Validate(InputTypes);
+
+		return bIsValid;
+	}
+
+	FOperatorHlsl* CreateDropoutOperator()
+	{
+		return new FDropout();
+	}
+
+	bool RegisterDropoutOperator(FOperatorRegistryHlsl& Registry)
+	{
+		Registry.OpAdd(TEXT("Dropout"), CreateDropoutOperator, ValidateDropoutOperator);
+		return true;
+	}
+} // UE::NNERuntimeRDG::Private::Hlsl
diff --git a/Engine/Plugins/Experimental/NNE/Source/NNERuntimeRDG/Private/Hlsl/NNERuntimeRDGDropout.h b/Engine/Plugins/Experimental/NNE/Source/NNERuntimeRDG/Private/Hlsl/NNERuntimeRDGDropout.h
new file mode 100644
--- /dev/null
+++ b/Engine/Plugins/Experimental/NNE/Source/NNERuntimeRDG/Private/Hlsl/NNERuntimeRDGDropout.h
@@ -0,0 +1,11 @@
+// Copyright Epic Games, Inc. All Rights Reserved.
+
+#pragma once
+
+#include "NNERuntimeRDG.h"
+#include "NNERuntimeRDGHlslOp.h"
+
+namespace UE::NNERuntimeRDG::Private::Hlsl
+{
+	bool RegisterDropoutOperator(FOperatorRegistryHlsl& Registry);
+} // UE::NNERuntimeRDG::Private::Hlsl
diff --git a/Engine/Plugins/Experimental/NNE/Source/NNERuntimeRDG/Private/Hlsl/NNERuntimeRDGElementWiseUnary.cpp b/Engine/Plugins/Experimental/NNE/Source/NNERuntimeRDG/Private/Hlsl/NNERuntimeRDGElementWiseUnary.cpp
--- a/Engine/Plugins/Experimental/NNE/Source/NNERuntimeRDG/Private/Hlsl/NNERuntimeRDGElementWiseUnary.cpp
+++ b/Engine/Plugins/Experimental/NNE/Source/NNERuntimeRDG/Private/Hlsl/NNERuntimeRDGElementWiseUnary.cpp
@@ -117,6 +117,16 @@ namespace UE::NNERuntimeRDG::Private::Hlsl
 	{
 	}
 
+	template<> bool TElementWiseUnary<NNECore::Internal::EElementWiseUnaryOperatorType::Clip>::Initialize(TConstArrayView<NNECore::FTensorDesc> InputTensorDescs, TConstArrayView<NNECore::FTensorDesc> OutputTensorDescs, const NNECore::FAttributeMap& Attributes)
+	{
+		check(InputTensorDescs.Num() == 1);
+		check(OutputTensorDescs.Num() == 1);
+
+		Alpha = Attributes.GetValueOrDefault(TEXT("min"), -3.402823e+38f);
+		Beta = Attributes.GetValueOrDefault(TEXT("max"), 3.402823e+38f);
+		return true;
+	}
+
 	template<NNECore::Internal::EElementWiseUnaryOperatorType OpType>
 	FOperatorHlsl* CreateElementWiseUnaryOperator()
 	{
@@ -208,6 +218,26 @@ namespace UE::NNERuntimeRDG::Private::Hlsl
 
 		return bIsValid;
 	}
+
+	template<>
+	bool ValidateElementWiseUnaryOperator<NNECore::Internal::EElementWiseUnaryOperatorType::Clip>(const NNECore::FAttributeMap& AttributeMap, TConstArrayView<ENNETensorDataType> InputTypes, TConstArrayView<NNECore::FSymbolicTensorShape> InputShapes)
+	{
+		bool bIsValid = true;
+
+		//version 6 of operator Clip, next version is 11
+		//[REDACTED_PROFILE_URL]/onnx/blob/main/docs/Changelog.md#Clip-6
+		FAttributeValidator AttributeValidator;
+		AttributeValidator.AddOptional(TEXT("min"), ENNEAttributeDataType::Float);
+		AttributeValidator.AddOptional(TEXT("max"), ENNEAttributeDataType::Float);
+		bIsValid &= AttributeValidator.Validate(AttributeMap);
+
+		FInputValidator InputValidator;
+		InputValidator.AddSupportedType(ENNETensorDataType::Float);
+		InputValidator.AddRequired();
+		bIsValid &= InputValidator.Validate(InputTypes);
+
+		return bIsValid;
+	}
 	
 	bool RegisterElementWiseUnaryOperators(FOperatorRegistryHlsl& Registry)
 	{
@@ -221,7 +251,7 @@ namespace UE::NNERuntimeRDG::Private::Hlsl
 		OP(Atanh);
 		//OP(BitShift);
 		OP(Ceil);
-		//OP(Clip);
+		OP(Clip);
 		OP(Cos);
 		OP(Cosh);
 		OP(Elu);
diff --git a/Engine/Plugins/Experimental/NNE/Source/NNERuntimeRDG/Private/Hlsl/NNERuntimeRDGFlatten.cpp b/Engine/Plugins/Experimental/NNE/Source/NNERuntimeRDG/Private/Hlsl/NNERuntimeRDGFlatten.cpp
new file mode 100644
--- /dev/null
+++ b/Engine/Plugins/Experimental/NNE/Source/NNERuntimeRDG/Private/Hlsl/NNERuntimeRDGFlatten.cpp
@@ -0,0 +1,130 @@
+// Copyright Epic Games, Inc. All Rights Reserved.
+
+#include "NNERuntimeRDGFlatten.h"
+#include "NNECoreTensor.h"
+#include "NNECoreTypes.h"
+#include "RenderGraphUtils.h"
+
+namespace UE::NNERuntimeRDG::Private::Hlsl
+{
+	DECLARE_GPU_STAT_NAMED(FNNEOperatorFlatten, TEXT("NNE.Operator.Hlsl.Flatten"));
+
+	/**
+	 * Flatten operator implementation
+	 */
+	class FFlatten : public FOperatorHlsl
+	{
+	public:
+
+		FFlatten() {}
+		virtual ~FFlatten() = default;
+
+		int32 Axis = 1;
+
+	public:
+
+		virtual int PrepareOutputs(TConstArrayView<NNECore::Internal::FTensorRef> InputTensors, TArrayView<NNECore::Internal::FTensorRef> OutputTensors) const override
+		{
+			check(InputTensors.Num() == 1);
+			check(OutputTensors.Num() == 1);
+
+			const NNECore::Internal::FTensor& X = *InputTensors[0];
+			TArray<uint32> OutputShapeData;
+			uint32 InnerDimSize = 1;
+
+			for (int32 i = 0; i < Axis; ++i)
+			{
+				InnerDimSize *= X.GetShape().GetData()[i];
+			}
+			OutputShapeData.Add(InnerDimSize);
+			OutputShapeData.Add(X.GetShape().Volume() / InnerDimSize);
+
+			NNECore::FTensorShape OutputShape = NNECore::FTensorShape::Make(OutputShapeData);
+
+			OutputTensors[0]->SetShape(OutputShape);
+			if (X.HasPreparedData())
+			{
+				OutputTensors[0]->SetPreparedData<uint8>(X.GetPreparedData<uint8>());
+			}
+			
+			return 0;
+		};
+
+		virtual bool Initialize(TConstArrayView<NNECore::FTensorDesc> InputTensorDescs, TConstArrayView<NNECore::FTensorDesc> OutputTensorDescs, const NNECore::FAttributeMap& Attributes) override
+		{
+			check(InputTensorDescs.Num() == 1);
+			check(OutputTensorDescs.Num() == 1);
+
+			int32 InputRank = InputTensorDescs[0].GetShape().Rank();
+			Axis = Attributes.GetValueOrDefault<int32>(TEXT("axis"), 1);
+
+			if (Axis > InputRank || Axis < -InputRank)
+			{
+				UE_LOG(LogNNE, Warning, TEXT("Flatten 'Axis' attribute should be in the range [-r,r] with r being the rank of the input (name: %s) however axis is %d while rank is %d."), *InputTensorDescs[0].GetName(), Axis, InputRank);
+				return false;
+			}
+			if (Axis < 0)
+			{
+				Axis = InputRank + Axis;
+			}
+			
+			return true;
+		}
+
+		virtual void Dispatch(FRDGBuilder& GraphBuilder, TConstArrayView<FTensorRDGRef> InputTensors, TConstArrayView<FTensorRDGRef> OutputTensors) override
+		{
+			check(InputTensors.Num() == 1);
+			check(OutputTensors.Num() == 1);
+			check(InputTensors[0] != nullptr);
+			check(OutputTensors[0] != nullptr);
+
+			const FTensorRDG& Data = *InputTensors[0];
+			const FTensorRDG& Output = *OutputTensors[0];
+
+			RDG_EVENT_SCOPE(GraphBuilder, "NNE.Operator.Hlsl.Flatten");
+			RDG_GPU_STAT_SCOPE(GraphBuilder, FNNEOperatorFlatten);
+
+			AddCopyBufferPass(GraphBuilder, Output.GetBuffer(), Data.GetBuffer());
+		}
+	};
+
+	bool ValidateFlattenOperator(const NNECore::FAttributeMap& AttributeMap, TConstArrayView<ENNETensorDataType> InputTypes, TConstArrayView<NNECore::FSymbolicTensorShape> InputFlattens)
+	{
+		bool bIsValid = true;
+
+		//This match version 13 of the Flatten operator
+		//[REDACTED_PROFILE_URL]/onnx/blob/main/docs/Operators.md#Flatten
+		FAttributeValidator AttributeValidator;
+		AttributeValidator.AddOptional(TEXT("axis"), ENNEAttributeDataType::Int32);
+		bIsValid &= AttributeValidator.Validate(AttributeMap);
+
+		FInputValidator InputValidator;
+		InputValidator.AddSupportedType(ENNETensorDataType::Half);
+		InputValidator.AddSupportedType(ENNETensorDataType::Float);
+		InputValidator.AddSupportedType(ENNETensorDataType::Double);
+		InputValidator.AddSupportedType(ENNETensorDataType::Int8);
+		InputValidator.AddSupportedType(ENNETensorDataType::Int16);
+		InputValidator.AddSupportedType(ENNETensorDataType::Int32);
+		InputValidator.AddSupportedType(ENNETensorDataType::Int64);
+		InputValidator.AddSupportedType(ENNETensorDataType::UInt8);
+		InputValidator.AddSupportedType(ENNETensorDataType::UInt16);
+		InputValidator.AddSupportedType(ENNETensorDataType::UInt32);
+		InputValidator.AddSupportedType(ENNETensorDataType::UInt64);
+		InputValidator.AddRequired();
+
+		bIsValid &= InputValidator.Validate(InputTypes);
+
+		return bIsValid;
+	}
+
+	FOperatorHlsl* CreateFlattenOperator()
+	{
+		return new FFlatten();
+	}
+
+	bool RegisterFlattenOperator(FOperatorRegistryHlsl& Registry)
+	{
+		Registry.OpAdd(TEXT("Flatten"), CreateFlattenOperator, ValidateFlattenOperator);
+		return true;
+	}
+} // UE::NNERuntimeRDG::Private::Hlsl
diff --git a/Engine/Plugins/Experimental/NNE/Source/NNERuntimeRDG/Private/Hlsl/NNERuntimeRDGFlatten.h b/Engine/Plugins/Experimental/NNE/Source/NNERuntimeRDG/Private/Hlsl/NNERuntimeRDGFlatten.h
new file mode 100644
--- /dev/null
+++ b/Engine/Plugins/Experimental/NNE/Source/NNERuntimeRDG/Private/Hlsl/NNERuntimeRDGFlatten.h
@@ -0,0 +1,11 @@
+// Copyright Epic Games, Inc. All Rights Reserved.
+
+#pragma once
+
+#include "NNERuntimeRDG.h"
+#include "NNERuntimeRDGHlslOp.h"
+
+namespace UE::NNERuntimeRDG::Private::Hlsl
+{
+	bool RegisterFlattenOperator(FOperatorRegistryHlsl& Registry);
+} // UE::NNERuntimeRDG::Private::Hlsl
diff --git a/Engine/Plugins/Experimental/NNE/Source/NNERuntimeRDG/Private/Hlsl/NNERuntimeRDGIdentity.cpp b/Engine/Plugins/Experimental/NNE/Source/NNERuntimeRDG/Private/Hlsl/NNERuntimeRDGIdentity.cpp
new file mode 100644
--- /dev/null
+++ b/Engine/Plugins/Experimental/NNE/Source/NNERuntimeRDG/Private/Hlsl/NNERuntimeRDGIdentity.cpp
@@ -0,0 +1,103 @@
+// Copyright Epic Games, Inc. All Rights Reserved.
+
+#include "NNERuntimeRDGIdentity.h"
+#include "NNECoreTensor.h"
+#include "NNECoreTypes.h"
+#include "RenderGraphUtils.h"
+
+namespace UE::NNERuntimeRDG::Private::Hlsl
+{
+	DECLARE_GPU_STAT_NAMED(FNNEOperatorIdentity, TEXT("NNE.Operator.Hlsl.Identity"));
+
+	/**
+	 * Identity operator implementation
+	 */
+	class FIdentity : public FOperatorHlsl
+	{
+	public:
+
+		FIdentity() {}
+		virtual ~FIdentity() = default;
+
+	public:
+
+		virtual int PrepareOutputs(TConstArrayView<NNECore::Internal::FTensorRef> InputTensors, TArrayView<NNECore::Internal::FTensorRef> OutputTensors) const override
+		{
+			check(InputTensors.Num() == 1);
+			check(OutputTensors.Num() == 1);
+
+			const NNECore::Internal::FTensor& X = *InputTensors[0];
+
+			OutputTensors[0]->SetShape(X.GetShape());
+			if (X.HasPreparedData())
+			{
+				OutputTensors[0]->SetPreparedData<uint8>(X.GetPreparedData<uint8>());
+			}
+			
+			return 0;
+		};
+
+		virtual bool Initialize(TConstArrayView<NNECore::FTensorDesc> InputTensorDescs, TConstArrayView<NNECore::FTensorDesc> OutputTensorDescs, const NNECore::FAttributeMap& Attributes) override
+		{
+			check(InputTensorDescs.Num() == 1);
+			check(OutputTensorDescs.Num() == 1);
+			
+			return true;
+		}
+
+		virtual void Dispatch(FRDGBuilder& GraphBuilder, TConstArrayView<FTensorRDGRef> InputTensors, TConstArrayView<FTensorRDGRef> OutputTensors) override
+		{
+			check(InputTensors.Num() == 1);
+			check(OutputTensors.Num() == 1);
+			check(InputTensors[0] != nullptr);
+			check(OutputTensors[0] != nullptr);
+
+			const FTensorRDG& Data = *InputTensors[0];
+			const FTensorRDG& Output = *OutputTensors[0];
+
+			RDG_EVENT_SCOPE(GraphBuilder, "NNE.Operator.Hlsl.Identity");
+			RDG_GPU_STAT_SCOPE(GraphBuilder, FNNEOperatorIdentity);
+
+			AddCopyBufferPass(GraphBuilder, Output.GetBuffer(), Data.GetBuffer());
+		}
+	};
+
+	bool ValidateIdentityOperator(const NNECore::FAttributeMap& AttributeMap, TConstArrayView<ENNETensorDataType> InputTypes, TConstArrayView<NNECore::FSymbolicTensorShape> InputIdentitys)
+	{
+		bool bIsValid = true;
+
+		//This match version 16 of the Identity operator
+		//[REDACTED_PROFILE_URL]/onnx/blob/main/docs/Operators.md#Identity
+		FAttributeValidator AttributeValidator;
+		bIsValid &= AttributeValidator.Validate(AttributeMap);
+
+		FInputValidator InputValidator;
+		InputValidator.AddSupportedType(ENNETensorDataType::Half);
+		InputValidator.AddSupportedType(ENNETensorDataType::Float);
+		InputValidator.AddSupportedType(ENNETensorDataType::Double);
+		InputValidator.AddSupportedType(ENNETensorDataType::Int8);
+		InputValidator.AddSupportedType(ENNETensorDataType::Int16);
+		InputValidator.AddSupportedType(ENNETensorDataType::Int32);
+		InputValidator.AddSupportedType(ENNETensorDataType::Int64);
+		InputValidator.AddSupportedType(ENNETensorDataType::UInt8);
+		InputValidator.AddSupportedType(ENNETensorDataType::UInt16);
+		InputValidator.AddSupportedType(ENNETensorDataType::UInt32);
+		InputValidator.AddSupportedType(ENNETensorDataType::UInt64);
+		InputValidator.AddRequired();
+
+		bIsValid &= InputValidator.Validate(InputTypes);
+
+		return bIsValid;
+	}
+
+	FOperatorHlsl* CreateIdentityOperator()
+	{
+		return new FIdentity();
+	}
+
+	bool RegisterIdentityOperator(FOperatorRegistryHlsl& Registry)
+	{
+		Registry.OpAdd(TEXT("Identity"), CreateIdentityOperator, ValidateIdentityOperator);
+		return true;
+	}
+} // UE::NNERuntimeRDG::Private::Hlsl
diff --git a/Engine/Plugins/Experimental/NNE/Source/NNERuntimeRDG/Private/Hlsl/NNERuntimeRDGIdentity.h b/Engine/Plugins/Experimental/NNE/Source/NNERuntimeRDG/Private/Hlsl/NNERuntimeRDGIdentity.h
new file mode 100644
--- /dev/null
+++ b/Engine/Plugins/Experimental/NNE/Source/NNERuntimeRDG/Private/Hlsl/NNERuntimeRDGIdentity.h
@@ -0,0 +1,11 @@
+// Copyright Epic Games, Inc. All Rights Reserved.
+
+#pragma once
+
+#include "NNERuntimeRDG.h"
+#include "NNERuntimeRDGHlslOp.h"
+
+namespace UE::NNERuntimeRDG::Private::Hlsl
+{
+	bool RegisterIdentityOperator(FOperatorRegistryHlsl& Registry);
+} // UE::NNERuntimeRDG::Private::Hlsl
diff --git a/Engine/Plugins/Experimental/NNE/Source/NNERuntimeRDG/Private/Hlsl/NNERuntimeRDGReshape.cpp b/Engine/Plugins/Experimental/NNE/Source/NNERuntimeRDG/Private/Hlsl/NNERuntimeRDGReshape.cpp
--- a/Engine/Plugins/Experimental/NNE/Source/NNERuntimeRDG/Private/Hlsl/NNERuntimeRDGReshape.cpp
+++ b/Engine/Plugins/Experimental/NNE/Source/NNERuntimeRDG/Private/Hlsl/NNERuntimeRDGReshape.cpp
@@ -4,6 +4,7 @@
 #include "NNECoreTensor.h"
 #include "NNECoreTypes.h"
 #include "RenderGraphUtils.h"
+#include "Helper/NNERuntimeRDGHelperReshape.h"
 
 namespace UE::NNERuntimeRDG::Private::Hlsl
 {
@@ -19,6 +20,8 @@ namespace UE::NNERuntimeRDG::Private::Hlsl
 		FReshape() {}
 		virtual ~FReshape() = default;
 
+		bool bAllowZero = false;
+
 	public:
 
 		virtual int PrepareOutputs(TConstArrayView<NNECore::Internal::FTensorRef> InputTensors, TArrayView<NNECore::Internal::FTensorRef> OutputTensors) const override
@@ -37,14 +40,10 @@ namespace UE::NNERuntimeRDG::Private::Hlsl
 				UE_LOG(LogNNE, Warning, TEXT("Reshape input 'Shape' (name: %s) should be constant for shape inference to succeed, however it is not."), *ShapeTensor.GetName());
 				return -1;
 			}
-			for (int64 Dim : ShapeTensor.GetPreparedData<int64>())
+
+			if (!ShapeHelper::Reshape::ReshapeTensor<int64>(X.GetShape(), ShapeTensor, bAllowZero, OutputShapeData))
 			{
-				if (Dim <= 0)
-				{
-					UE_LOG(LogNNE, Warning, TEXT("Cannot reshape 'Data' (name: %s) to a shape with a dimension less or equal to 0."), *X.GetName());
-					return -1;
-				}
-				OutputShapeData.Add((uint32)Dim);
+				return -1;
 			}
 
 			NNECore::FTensorShape OutputShape = NNECore::FTensorShape::Make(OutputShapeData);
@@ -62,6 +61,8 @@ namespace UE::NNERuntimeRDG::Private::Hlsl
 		{
 			check(InputTensorDescs.Num() == 2);
 			check(OutputTensorDescs.Num() == 1);
+
+			bAllowZero = (bool) (Attributes.GetValueOrDefault<int32>(TEXT("allowzero"), 0));
 			
 			return true;
 		}
@@ -87,9 +88,10 @@ namespace UE::NNERuntimeRDG::Private::Hlsl
 	{
 		bool bIsValid = true;
 
-		//This match version 5 of the Reshape operator, next version is 13
+		//This match version 14 of the Reshape operator
 		//[REDACTED_PROFILE_URL]/onnx/blob/main/docs/Operators.md#Reshape
 		FAttributeValidator AttributeValidator;
+		AttributeValidator.AddOptional(TEXT("allowzero"), ENNEAttributeDataType::Int32);
 		bIsValid &= AttributeValidator.Validate(AttributeMap);
 
 		FInputValidator InputValidator;
diff --git a/Engine/Plugins/Experimental/NNE/Source/NNERuntimeRDG/Private/NNERuntimeRDGHlsl.cpp b/Engine/Plugins/Experimental/NNE/Source/NNERuntimeRDG/Private/NNERuntimeRDGHlsl.cpp
--- a/Engine/Plugins/Experimental/NNE/Source/NNERuntimeRDG/Private/NNERuntimeRDGHlsl.cpp
+++ b/Engine/Plugins/Experimental/NNE/Source/NNERuntimeRDG/Private/NNERuntimeRDGHlsl.cpp
@@ -8,15 +8,19 @@
 #include "NNERuntimeRDGHlslHelper.h"
 #include "NNERuntimeRDGModelHlsl.h"
 #include "NNEUtilsModelOptimizer.h"
+#include "Hlsl/NNERuntimeRDGBatchNormalization.h"
 #include "Hlsl/NNERuntimeRDGCast.h"
 #include "Hlsl/NNERuntimeRDGConv.h"
 #include "Hlsl/NNERuntimeRDGConcat.h"
 #include "Hlsl/NNERuntimeRDGConvTranspose.h"
+#include "Hlsl/NNERuntimeRDGDropout.h"
 #include "Hlsl/NNERuntimeRDGElementWiseBinary.h"
 #include "Hlsl/NNERuntimeRDGElementWiseUnary.h"
 #include "Hlsl/NNERuntimeRDGElementWiseVariadic.h"
+#include "Hlsl/NNERuntimeRDGFlatten.h"
 #include "Hlsl/NNERuntimeRDGGather.h"
 #include "Hlsl/NNERuntimeRDGGemm.h"
+#include "Hlsl/NNERuntimeRDGIdentity.h"
 #include "Hlsl/NNERuntimeRDGInstanceNormalization.h"
 #include "Hlsl/NNERuntimeRDGPad.h"
 #include "Hlsl/NNERuntimeRDGReshape.h"
@@ -38,15 +42,19 @@ bool UNNERuntimeRDGHlslImpl::Init()
 	FOperatorRegistryHlsl* registry = FOperatorRegistryHlsl::Get();
 	check(registry != nullptr);
 
+	RegisterBatchNormalizationOperator(*registry);
 	RegisterCastOperator(*registry);
 	RegisterConvOperator(*registry);
 	RegisterConcatOperator(*registry);
 	RegisterConvTransposeOperator(*registry);
+	RegisterDropoutOperator(*registry);
 	RegisterElementWiseBinaryOperators(*registry);
 	RegisterElementWiseUnaryOperators(*registry);
 	RegisterElementWiseVariadicOperators(*registry);
+	RegisterFlattenOperator(*registry);
 	RegisterGatherOperator(*registry);
 	RegisterGemmOperator(*registry);
+	RegisterIdentityOperator(*registry);
 	RegisterInstanceNormalizationOperator(*registry);
 	RegisterPadOperator(*registry);
 	RegisterReshapeOperator(*registry);
diff --git a/Engine/Plugins/Experimental/NNE/Source/NNEUtils/Private/NNEUtilsModelBuilderONNX.cpp b/Engine/Plugins/Experimental/NNE/Source/NNEUtils/Private/NNEUtilsModelBuilderONNX.cpp
--- a/Engine/Plugins/Experimental/NNE/Source/NNEUtils/Private/NNEUtilsModelBuilderONNX.cpp
+++ b/Engine/Plugins/Experimental/NNE/Source/NNEUtils/Private/NNEUtilsModelBuilderONNX.cpp
@@ -350,13 +350,14 @@ NNEUTILS_API bool CreateONNXModelForOperator(bool UseVariadicShapeForModel, cons
 	
 	int64 IrVersion = OnnxIrVersion;
 	int64 OpsetVersion = OnnxOpsetVersion;
-	if (OperatorName == TEXT("Upsample") ||	// deprecated starting opset 10
-		OperatorName == TEXT("Pad") ||		// current implementation is opset 2 (next version is 11)
-		OperatorName == TEXT("Slice") ||	// current implementation is opset 1 (next version is 10)
-		OperatorName == TEXT("Shape") ||	// current implementation is opset 1 (next version is 13)
-		OperatorName == TEXT("Reshape") ||	// current implementation is opset 5 (next version is 13)
-		OperatorName == TEXT("Squeeze") ||	// current implementation is opset 1 (next version is 11)
-		OperatorName == TEXT("Unsqueeze"))	// current implementation is opset 1 (next version is 11)
+	if (OperatorName == TEXT("BatchNormalization") ||	// current implementation is opset 9 (next version is 14)
+		OperatorName == TEXT("Clip") ||					// current implementation is opset 6 (next version is 11)
+		OperatorName == TEXT("Pad") ||					// current implementation is opset 2 (next version is 11)
+		OperatorName == TEXT("Shape") ||				// current implementation is opset 1 (next version is 13)
+		OperatorName == TEXT("Slice") ||				// current implementation is opset 1 (next version is 10)
+		OperatorName == TEXT("Squeeze") ||				// current implementation is opset 1 (next version is 11)
+		OperatorName == TEXT("Unsqueeze") ||			// current implementation is opset 1 (next version is 11)
+		OperatorName == TEXT("Upsample"))				// deprecated starting opset 10
 	{
 		OpsetVersion = 9;
 	}
