NNE DML - Fix another constant cpu indices issue
#rnx

diff --git a/Engine/Plugins/Experimental/NNE/Source/NNERuntimeRDG/Private/Dml/NNEDmlModel.cpp b/Engine/Plugins/Experimental/NNE/Source/NNERuntimeRDG/Private/Dml/NNEDmlModel.cpp
--- a/Engine/Plugins/Experimental/NNE/Source/NNERuntimeRDG/Private/Dml/NNEDmlModel.cpp
+++ b/Engine/Plugins/Experimental/NNE/Source/NNERuntimeRDG/Private/Dml/NNEDmlModel.cpp
@@ -1019,15 +1019,15 @@ bool FModel::InitCompiledOp(TConstArrayView<int32> OpInputIndices, uint64 Tensor
 				uint8*			UploadBuffPtr = static_cast<uint8*>(RHICmdList.LockBuffer(UploadBuff, 0, TensorDataSize, RLM_WriteOnly_NoOverwrite));
 				uint64			UploadOffset = 0;
 
-				int32 WeightIdx = 0;
-				for (int32 TensorIdx : WeightTensorIndices)
+				for (int32 WeightIdx = 0; WeightIdx < WeightTensorIndices.Num(); ++WeightIdx)
 				{
+					int32 TensorIdx = WeightTensorIndices[WeightIdx];
 					if (ConstantCPUTensorIndices.Find(TensorIdx) != -1)
 					{
 						continue;
 					}
 					
-					const FTensorRDG&		Tensor = WeightTensorRDGs[WeightIdx++];
+					const FTensorRDG&		Tensor = WeightTensorRDGs[WeightIdx];
 					TConstArrayView<uint8>	TensorData = Tensor.GetPreparedData<uint8>();
 
 					FBufferRHIRef WeightBuff;
