NNE DML - Fix MatMul handling of leading ones
#rnx

diff --git a/Engine/Plugins/Experimental/NNE/Source/NNERuntimeRDG/Private/Dml/NNEDmlOperatorMatMul.cpp b/Engine/Plugins/Experimental/NNE/Source/NNERuntimeRDG/Private/Dml/NNEDmlOperatorMatMul.cpp
--- a/Engine/Plugins/Experimental/NNE/Source/NNERuntimeRDG/Private/Dml/NNEDmlOperatorMatMul.cpp
+++ b/Engine/Plugins/Experimental/NNE/Source/NNERuntimeRDG/Private/Dml/NNEDmlOperatorMatMul.cpp
@@ -32,10 +32,22 @@ public:
 		const NNECore::Internal::FTensor& BTensor = InputTensors[1];
 		const NNECore::Internal::FTensor& OutTensor = OutputTensors[0];
 
+		auto TrimMatrixStackLeadingOnes = [](TConstArrayView<uint32>&& Input) -> TConstArrayView<uint32>
+			{
+				int Start = 0;
+				for(; Start < Input.Num() - 2; ++Start)
+				{
+					if(Input[Start] != 1)
+					{
+						break;
+					}
+				}
+				return Input.RightChop(Start);
+			};
 
-		DmlUtil::FSmallUIntArray ActualShapeA (ATensor.GetShape().GetData());
-		DmlUtil::FSmallUIntArray ActualShapeB (BTensor.GetShape().GetData());
-		DmlUtil::FSmallUIntArray ActualShapeOut (OutTensor.GetShape().GetData());
+		DmlUtil::FSmallUIntArray ActualShapeA (TrimMatrixStackLeadingOnes(ATensor.GetShape().GetData()));
+		DmlUtil::FSmallUIntArray ActualShapeB (TrimMatrixStackLeadingOnes(BTensor.GetShape().GetData()));
+		DmlUtil::FSmallUIntArray ActualShapeOut (TrimMatrixStackLeadingOnes(OutTensor.GetShape().GetData()));
 
 		check(ActualShapeA.Num() > 0);
 		check(ActualShapeB.Num() > 0);
