NNE DML - Add NeuralMorphModel_Local test with simplified version of model
#rnx

diff --git a/Engine/Plugins/Experimental/NNE/Source/NNEQA/Resources/NNEQATestDesc.json b/Engine/Plugins/Experimental/NNE/Source/NNEQA/Resources/NNEQATestDesc.json
--- a/Engine/Plugins/Experimental/NNE/Source/NNEQA/Resources/NNEQATestDesc.json
+++ b/Engine/Plugins/Experimental/NNE/Source/NNEQA/Resources/NNEQATestDesc.json
@@ -55,6 +55,15 @@
 						{ "name": "NNERuntimeRDGDml",  "skip": true, "comment": "Model not supported yet." }
 					]
 				},
+				{
+					"target": "NeuralMorphModel_Local-sim",
+					"tags": [ "SmokeTest" ],
+					"runtimes": [
+						{ "name": "NNERuntimeORTDml", "absolute_tolerance": 1e-6 },
+						{ "name": "NNERuntimeRDGHlsl",  "skip": true, "comment": "Missing Reshape operator." },
+						{ "name": "NNERuntimeRDGDml", "absolute_tolerance": 1e-6 }
+					]
+				},
 				{
 					"target": "VertexDeltaModel",
 					"absolute_tolerance": 1e-4
