NNE Update style transfer demo with models trained for longer + fixes in test json desc.

Tests:
* style transfer demo on all model editor win64
* model tests editor win64

#rnx

diff --git a/Engine/Plugins/Experimental/NNE/Source/NNEQA/Resources/NNEQATestDesc.json b/Engine/Plugins/Experimental/NNE/Source/NNEQA/Resources/NNEQATestDesc.json
--- a/Engine/Plugins/Experimental/NNE/Source/NNEQA/Resources/NNEQATestDesc.json
+++ b/Engine/Plugins/Experimental/NNE/Source/NNEQA/Resources/NNEQATestDesc.json
@@ -6,7 +6,7 @@
 			"is_model_test": true,
 			"targets": [
 				{
-					"target": "udnie-20-small-720p",
+					"target": "udnie-160-small-720p",
 					"runtimes": [
 						{ "name": "NNERuntimeRDGDml",  "skip": true, "comment": "Model not supported yet." }
 					],
@@ -19,7 +19,7 @@
 			"is_model_test": true,
 			"targets": [
 				{
-					"target": "udnie-20-small-dyn",
+					"target": "udnie-160-small-dyn",
 					"runtimes": [
 						{ "name": "NNERuntimeRDGDml",  "skip": true, "comment": "Model not supported yet." }
 					],
@@ -41,28 +41,23 @@
 					"target": "NeuralMorphModel_global_Dax",
 					"tags": [ "SmokeTest" ],
 					"runtimes": [
-						{ "name": "NNERuntimeORTDml", "skip": true, "comment": "" },
-						{ "name": "NNERuntimeRDGHlsl",	"skip": true, "comment": "Model not supported yet."	},
-						{ "name": "NNERuntimeRDGDml",  "skip": true, "comment": "Model not supported yet." }
+						{ "name": "NNERuntimeORTDml", "relative_tolerance": 2e-4 },
+						{ "name": "NNERuntimeRDGHlsl", "absolute_tolerance": 1e-6 },
+						{ "name": "NNERuntimeRDGDml", "skip": true, "comment": "Model not supported yet." }
 					]
 				},
 				{
 					"target": "NeuralMorphModel_Local",
 					"tags": [ "SmokeTest" ],
+					"absolute_tolerance": 1e-6,
 					"runtimes": [
-						{ "name": "NNERuntimeORTDml", "absolute_tolerance": 1e-6 },
-						{ "name": "NNERuntimeRDGHlsl",  "skip": true, "comment": "Missing Reshape operator." },
-						{ "name": "NNERuntimeRDGDml",  "skip": true, "comment": "Model not supported yet." }
+						{ "name": "NNERuntimeRDGDml", "skip": true, "comment": "Model not supported yet." }
 					]
 				},
 				{
 					"target": "NeuralMorphModel_Local-sim",
 					"tags": [ "SmokeTest" ],
-					"runtimes": [
-						{ "name": "NNERuntimeORTDml", "absolute_tolerance": 1e-6 },
-						{ "name": "NNERuntimeRDGHlsl",  "skip": true, "comment": "Missing Reshape operator." },
-						{ "name": "NNERuntimeRDGDml", "absolute_tolerance": 1e-6 }
-					]
+					"absolute_tolerance": 1e-6
 				},
 				{
 					"target": "VertexDeltaModel",
