RDG DML graph builder now supports one node to multiple nodes intermediate/output connections

#rnx

diff --git a/Engine/Plugins/Experimental/NNE/Source/NNEQA/Resources/NNEQATestDesc.json b/Engine/Plugins/Experimental/NNE/Source/NNEQA/Resources/NNEQATestDesc.json
--- a/Engine/Plugins/Experimental/NNE/Source/NNEQA/Resources/NNEQATestDesc.json
+++ b/Engine/Plugins/Experimental/NNE/Source/NNEQA/Resources/NNEQATestDesc.json
@@ -7,9 +7,6 @@
 			"targets": [
 				{
 					"target": "udnie-160-small-720p",
-					"runtimes": [
-						{ "name": "NNERuntimeRDGDml",  "skip": true, "comment": "Model not supported yet." }
-					],
 					"absolute_tolerance": 5.0e-2
 				}
 			]
@@ -41,9 +38,9 @@
 					"target": "NeuralMorphModel_global_Dax",
 					"tags": [ "SmokeTest" ],
 					"runtimes": [
-						{ "name": "NNERuntimeORTDml", "relative_tolerance": 2e-4 },
-						{ "name": "NNERuntimeRDGHlsl", "absolute_tolerance": 1e-6 },
-						{ "name": "NNERuntimeRDGDml", "skip": true, "comment": "Model not supported yet." }
+						{ "name": "NNERuntimeORTDml", "absolute_tolerance": 1e-4 },
+						{ "name": "NNERuntimeRDGDml", "absolute_tolerance": 1e-4 },
+						{ "name": "NNERuntimeRDGHlsl", "absolute_tolerance": 1e-6 }
 					]
 				},
 				{
diff --git a/Engine/Plugins/Experimental/NNE/Source/NNERuntimeRDG/Private/Dml/NNEDmlModel.cpp b/Engine/Plugins/Experimental/NNE/Source/NNERuntimeRDG/Private/Dml/NNEDmlModel.cpp
--- a/Engine/Plugins/Experimental/NNE/Source/NNERuntimeRDG/Private/Dml/NNEDmlModel.cpp
+++ b/Engine/Plugins/Experimental/NNE/Source/NNERuntimeRDG/Private/Dml/NNEDmlModel.cpp
@@ -418,10 +418,6 @@ public:
 		{
 			if (Edge.Type == EEdgeType::Input)
 			{
-				check(Edge.NodeSrcOutput >= 0);
-				check(Edge.NodeDst >= 0);
-				check(Edge.NodeDstInput >= 0);
-
 				DML_INPUT_GRAPH_EDGE_DESC& Input = InputEdges.Add_GetRef({});
 
 				Input.GraphInputIndex = Edge.NodeSrcOutput;
@@ -433,10 +429,6 @@ public:
 			}
 			else if (Edge.Type == EEdgeType::Output)
 			{
-				check(Edge.NodeDstInput >= 0);
-				check(Edge.NodeSrc >= 0);
-				check(Edge.NodeSrcOutput >= 0);
-
 				DML_OUTPUT_GRAPH_EDGE_DESC&	Output = OutputEdges.Add_GetRef({});
 
 				Output.GraphOutputIndex = Edge.NodeDstInput;
@@ -450,11 +442,6 @@ public:
 			{
 				DML_INTERMEDIATE_GRAPH_EDGE_DESC& Intermediate = IntermediateEdges.Add_GetRef({});
 
-				check(Edge.NodeSrc >= 0);
-				check(Edge.NodeSrcOutput >= 0);
-				check(Edge.NodeDst >= 0);
-				check(Edge.NodeDstInput >= 0);
-
 				Intermediate.FromNodeIndex = Edge.NodeSrc;
 				Intermediate.FromNodeOutputIndex = Edge.NodeSrcOutput;
 				Intermediate.ToNodeIndex = Edge.NodeDst;
@@ -486,7 +473,7 @@ public:
 
 			Edge.Type = DML_GRAPH_EDGE_TYPE_INPUT;
 			Edge.Desc = &InputEdges[Idx];
-		}			
+		}
 
 		for (int32 Idx = 0; Idx < OutputEdges.Num(); ++Idx)
 		{
@@ -535,145 +522,235 @@ private:
 	bool AddEdges(const FGraphDesc& InGraph)
 	{
 		Edges.Reset();
+		TensorInConnCounts.Reset();
+		TensorOutConnCounts.Reset();
 		Operators.Reset();
 		NumInputs = 0;
 		NumOutputs = 0;
 
-		// Loop through all operators that have this tensor as an input/weight and it as an input edge
-		auto AddInputs = [this, InGraph](TConstArrayView<int32> InTensorIndices) -> bool
+		TensorInConnCounts.SetNumZeroed(InGraph.AllTensors.Num());
+		TensorOutConnCounts.SetNumZeroed(InGraph.AllTensors.Num());
+		Operators.Reset(InGraph.Operators.Num());
+
+		for (int32 Idx : InGraph.OpInputIndices)
 		{
-			for (int32 Idx = 0; Idx < InTensorIndices.Num(); ++Idx)
-			{
-				const int32 TensorIdx = InTensorIndices[Idx];
+			TensorInConnCounts[Idx] += 1;
+		}
 
-				for (const FOpDesc& OpDesc : InGraph.Operators)
-				{
-					for (int32 InIdx = 0; InIdx < OpDesc.InputCount; ++InIdx)
-					{
-						const int32 OpInputIdx = InGraph.OpInputIndices[InIdx + OpDesc.InputStart];
+		for (int32 Idx : InGraph.OpOutputIndices)
+		{
+			TensorOutConnCounts[Idx] += 1;
+		}
 
-						if (OpInputIdx == TensorIdx)
-						{
-							bool bRes = AddInput(TensorIdx);
-							if (!bRes)
-							{
-								return bRes;
-							}
-						}
-					}
-				}
-			}
+		if (!AddInputEdges(InGraph.InputIndices, InGraph))
+		{
+			UE_LOG(LogNNE, Warning, TEXT("DMLGraphBuilder failed to add input tensors"));
+			return false;
+		}
+		
+		if (!AddInputEdges(InGraph.WeightIndices, InGraph))
+		{
+			UE_LOG(LogNNE, Warning, TEXT("DMLGraphBuilder failed to add weight tensors"));
+			return false;
+		}
+		
+		if (!AddOutputEdges(InGraph))
+		{
+			UE_LOG(LogNNE, Warning, TEXT("DMLGraphBuilder failed to add output tensors"));
+			return false;
+		}
 
-			return true;
-		};
+		if (!AddIntermediateEdges(InGraph))
+		{
+			UE_LOG(LogNNE, Warning, TEXT("DMLGraphBuilder failed to add intermediate tensors"));
+			return false;
+		}
 
-		if (!AddInputs(InGraph.InputIndices))
+		if (!AddOperators(InGraph))
 		{
-			UE_LOG(LogNNE, Warning, TEXT("DMLGraphBuilder failed to process input tensors"));
+			UE_LOG(LogNNE, Warning, TEXT("DMLGraphBuilder failed to add operators"));
 			return false;
 		}
+
+		// Validate edges
+		for (const FEdge& Edge : Edges)
+		{
+			if (Edge.Type == EEdgeType::Input)
+			{
+				if (Edge.NodeSrcOutput == -1 || Edge.NodeDst == -1 || Edge.NodeDstInput == -1)
+				{
+					UE_LOG(LogNNE, Warning, TEXT("DMLGraphBuilder error invalid input graph edge detected for tensor:%s"), *InGraph.AllTensors[Edge.TensorIdx].GetName());
+					return false;
+				}
+			}
+			else if (Edge.Type == EEdgeType::Output)
+			{
+				if (Edge.NodeDstInput == -1 || Edge.NodeSrc == -1 || Edge.NodeSrcOutput == -1)
+				{
+					UE_LOG(LogNNE, Warning, TEXT("DMLGraphBuilder error invalid output graph edge detected for tensor:%s"), *InGraph.AllTensors[Edge.TensorIdx].GetName());
+					return false;
+				}
+			}
+			else if (Edge.Type == EEdgeType::Intermediate)
+			{
+				if (Edge.NodeSrc == -1 || Edge.NodeSrcOutput == -1 || Edge.NodeDst == -1 || Edge.NodeDstInput == -1)
+				{
+					UE_LOG(LogNNE, Warning, TEXT("DMLGraphBuilder error invalid intermediate graph edge detected for tensor:%s"), *InGraph.AllTensors[Edge.TensorIdx].GetName());
+					return false;
+				}
+			}
+		}
 		
-		if (!AddInputs(InGraph.WeightIndices))
+		return true;
+	}
+
+	bool AddInputEdges(TConstArrayView<int32> InTensorIndices, const FGraphDesc& InGraph)
+	{
+		for (int32 Idx = 0; Idx < InTensorIndices.Num(); ++Idx)
 		{
-			UE_LOG(LogNNE, Warning, TEXT("DMLGraphBuilder failed to process weight tensors"));
-			return false;
+			const int32 TensorIdx = InTensorIndices[Idx];
+			const int32 ConnCount = TensorInConnCounts[TensorIdx];
+
+			if (ConnCount == 0)
+			{
+				if (InGraph.ConstantCPUIndices.Find(TensorIdx) == INDEX_NONE)
+				{
+					UE_LOG(LogNNE, Warning, TEXT("DmlGraphBuilder tensor:%d has no connections"), *InGraph.AllTensors[TensorIdx].GetName());
+					return false;
+				}
+			}
+
+			for (int32 ConnIdx = 0; ConnIdx < ConnCount; ++ConnIdx)
+			{
+				AddInputEdge(TensorIdx);
+			}
 		}
 
+		return true;
+	}
+
+	bool AddOutputEdges(const FGraphDesc& InGraph)
+	{
 		for (int32 Idx = 0; Idx < InGraph.OutputIndices.Num(); ++Idx)
 		{
 			const int32 TensorIdx = InGraph.OutputIndices[Idx];
-			
-			if (!AddOutput(TensorIdx))
+			const int32 ConnCount = TensorOutConnCounts[TensorIdx];
+
+			if (ConnCount == 0)
 			{
-				UE_LOG(LogNNE, Warning, TEXT("DMLGraphBuilder failed to process output tensors"));
+				UE_LOG(LogNNE, Warning, TEXT("DmlGraphBuilder tensor:%d has no connections"), *InGraph.AllTensors[TensorIdx].GetName());
 				return false;
 			}
+
+			for (int32 ConnIdx = 0; ConnIdx < ConnCount; ++ConnIdx)
+			{
+				AddOutputEdge(TensorIdx);
+			}
 		}
 
-		Operators.Reset(InGraph.Operators.Num());
-		for (const FOpDesc& OpDesc : InGraph.Operators)
+		return true;
+	}
+
+	bool AddIntermediateEdges(const FGraphDesc& InGraph)
+	{
+		for (int32 Idx = 0; Idx < InGraph.IntermediateIndices.Num(); ++Idx)
 		{
-			if (!AddOp(OpDesc, InGraph))
+			const int32 TensorIdx = InGraph.IntermediateIndices[Idx];
+			const int32 InputConnCount = TensorInConnCounts[TensorIdx];
+			const int32 OutputConnCount = TensorOutConnCounts[TensorIdx];
+
+			if (InputConnCount == 0 || OutputConnCount == 0)
 			{
-				UE_LOG(LogNNE, Warning, TEXT("DMLGraphBuilder failed to process operator tensors"));
+				UE_LOG(LogNNE, Warning, TEXT("DmlGraphBuilder tensor:%d has no connections"), *InGraph.AllTensors[TensorIdx].GetName());
 				return false;
 			}
+
+			for (int32 ConnIdx = 0; ConnIdx < InputConnCount; ++ConnIdx)
+			{
+				AddIntermediateEdge(TensorIdx, -1, -1);
+			}
 		}
 
-		//Note: TODO: Validate edges here
-		
 		return true;
 	}
 
-	bool AddInput(int32 TensorIdx)
+	bool AddOperators(const FGraphDesc& InGraph)
 	{
-		bool bRes = AddEdge(
+		for (const FOpDesc& CurrOp : InGraph.Operators) 
+		{
+			DML_OPERATOR_GRAPH_NODE_DESC& OpDesc = Operators.Add_GetRef({});
+
+			OpDesc.Operator = CurrOp.Op->GetOperator();
+			OpDesc.Name = CurrOp.DbgName.Get();
+
+			const int32 NodeIdx = Operators.Num() - 1;
+
+			for (int32 Idx = 0; Idx < CurrOp.InputCount; ++Idx)
+			{
+				const int32 TensorIdx = InGraph.OpInputIndices[Idx + CurrOp.InputStart];
+
+				if (!ConnectEdgeDst(TensorIdx, NodeIdx, Idx))
+				{
+					UE_LOG(LogNNE, Warning, TEXT("DmlGraphBuilder failed to connect intermediate edge dst for tensor:%s"), *InGraph.AllTensors[TensorIdx].GetName());
+					return false;
+				}
+			}
+
+			for (int32 Idx = 0; Idx < CurrOp.OutputCount; ++Idx)
+			{
+				const int32 TensorIdx = InGraph.OpOutputIndices[Idx + CurrOp.OutputStart];
+				int32 ConnCount = TensorInConnCounts[TensorIdx];
+
+				if (ConnCount == 0)
+				{
+					// This is for output edges, they don't have input connections
+					ConnCount = TensorOutConnCounts[TensorIdx];
+				}
+
+				for (int32 ConnIdx = 0; ConnIdx < ConnCount; ++ConnIdx)
+				{
+					if (!ConnectEdgeSrc(TensorIdx, NodeIdx, Idx))
+					{
+						UE_LOG(LogNNE, Warning, TEXT("DmlGraphBuilder failed to connect intermediate edge src for tensor:%s"), *InGraph.AllTensors[TensorIdx].GetName());
+						return false;
+					}
+				}
+			}
+		}
+
+		return true;
+	}
+
+	void AddInputEdge(int32 TensorIdx)
+	{
+		Edges.Emplace(
 			FEdge(EEdgeType::Input)
 				.SetTensorIdx(TensorIdx)
-				.SetNodeSrcOutput(NumInputs),
-			false
+				.SetNodeSrcOutput(NumInputs)
 		);
 
 		++NumInputs;
-		return bRes;
 	}
 
-	bool AddOutput(int32 TensorIdx)
+	void AddOutputEdge(int32 TensorIdx)
 	{
-		bool bRes = AddEdge(
+		Edges.Emplace(
 			FEdge(EEdgeType::Output)
 				.SetTensorIdx(TensorIdx)
 				.SetNodeDstInput(NumOutputs)
 		);
 
 		++NumOutputs;
-		return bRes;
 	}
 
-	bool AddIntermediate(int32 TensorIdx, int32 NodeSrc, int32 NodeSrcOutput)
+	void AddIntermediateEdge(int32 TensorIdx, int32 NodeSrc, int32 NodeSrcOutput)
 	{
-		FEdge* ConnEdge = Edges.FindByPredicate(
-				[TensorIdx](const FEdge& Curr)
-				{
-					return Curr.TensorIdx == TensorIdx;
-				}
-		);
-
-		if (ConnEdge)
-		{
-			return ConnectEdgeSrc(TensorIdx, NodeSrc, NodeSrcOutput);
-		}
-		else
-		{
-			return AddEdge(
-				FEdge(EEdgeType::Intermediate)
-					.SetTensorIdx(TensorIdx)
-					.SetNodeSrc(NodeSrc)
-					.SetNodeSrcOutput(NodeSrcOutput)
+		Edges.Emplace(
+			FEdge(EEdgeType::Intermediate)
+				.SetTensorIdx(TensorIdx)
+				.SetNodeSrc(NodeSrc)
+				.SetNodeSrcOutput(NodeSrcOutput)
 			);
-		}
-	}
-
-	bool AddEdge(const FEdge& Edge, bool bIsUnique = true)
-	{
-		FEdge* StartEdge = Edges.FindByPredicate(
-				[Edge](const FEdge& Curr)
-				{
-					return Curr.TensorIdx == Edge.TensorIdx;
-				}
-		);
-
-		if (bIsUnique)
-		{
-			if (StartEdge != nullptr)
-			{
-				UE_LOG(LogNNE, Warning, TEXT("DmlGraphBuilder:Edge is not unique"));
-				return false;
-			}
-		}
-
-		Edges.Add(Edge);
-		return true;
 	}
 
 	bool ConnectEdgeDst(int32 TensorIdx, int32 NodeDst, int32 NodeDstInput)
@@ -696,11 +773,6 @@ private:
 				bFoundEdge = true;
 				break;
 			}
-			else if (Curr->NodeDst == NodeDst && Curr->NodeDstInput == NodeDstInput)
-			{
-				bFoundEdge = true;
-				break;
-			}
 		}
 
 		checkf(bFoundEdge, TEXT("ConnectEdgeDst() has failed"));
@@ -728,52 +800,15 @@ private:
 				bFoundEdge = true;
 				break;
 			}
-			else if (Curr->NodeSrc == NodeSrc && Curr->NodeSrcOutput == NodeSrcOutput)
-			{
-				bFoundEdge = true;
-				break;
-			}
 		}
 
 		checkf(bFoundEdge, TEXT("ConnectEdgeSrc() has failed"));
 		return bFoundEdge;
 	}
 
-	bool AddOp(const FOpDesc& InOp, const FGraphDesc& InGraph)
-	{
-		DML_OPERATOR_GRAPH_NODE_DESC& OpDesc = Operators.Add_GetRef({});
-
-		OpDesc.Operator = InOp.Op->GetOperator();
-		OpDesc.Name = InOp.DbgName.Get();
-
-		const int32 NodeIdx = Operators.Num() - 1;
-
-		for (int32 Idx = 0; Idx < InOp.InputCount; ++Idx)
-		{
-			const int32 TensorIdx = InGraph.OpInputIndices[Idx + InOp.InputStart];				
-			
-			bool bRes = ConnectEdgeDst(TensorIdx, NodeIdx, Idx);
-			if (!bRes)
-			{
-				return bRes;
-			}
-		}
-
-		for (int32 Idx = 0; Idx < InOp.OutputCount; ++Idx)
-		{
-			const int32 TensorIdx = InGraph.OpOutputIndices[Idx + InOp.OutputStart];
-
-			bool bRes = AddIntermediate(TensorIdx, NodeIdx, Idx);
-			if (!bRes)
-			{
-				return bRes;
-			}
-		}
-
-		return true;
-	}
-
 	TArray<FEdge>							Edges;
+	TArray<int32>							TensorInConnCounts;
+	TArray<int32>							TensorOutConnCounts;
 	TArray<DML_OPERATOR_GRAPH_NODE_DESC>	Operators;
 	int32									NumInputs;
 	int32									NumOutputs;
@@ -900,7 +935,7 @@ bool FModel::Init(TConstArrayView<uint8> ModelData, FDmlDeviceContext* InDevCtx)
 			OpInputTensors.RemoveAt(ConstIdx);
 			OpInputIndices.RemoveAt(OpDesc.InputStart + ConstIdx);
 
-			ConstantCPUTensorIndices.Emplace(TensorIdx);
+			ConstantCPUTensorIndices.AddUnique(TensorIdx);
 		}
 
 		OpDesc.InputCount = OpInputTensors.Num();
@@ -1003,7 +1038,7 @@ bool FModel::InitCompiledOp(TConstArrayView<int32> OpInputIndices, uint64 Tensor
 
 			for (int32 InputIdx : InputTensorIndices)
 			{
-				if (ConstantCPUTensorIndices.Find(InputIdx) == -1)
+				if (ConstantCPUTensorIndices.Find(InputIdx) == INDEX_NONE)
 				{
 					Inputs.Emplace(nullptr);
 				}
@@ -1031,7 +1066,7 @@ bool FModel::InitCompiledOp(TConstArrayView<int32> OpInputIndices, uint64 Tensor
 					TConstArrayView<uint8>	TensorData = Tensor.GetPreparedData<uint8>();
 
 					FBufferRHIRef WeightBuff;
-					WeightBuff = CreateRHIBuffer(RHICmdList, TensorData.Num(), WeightBuffUsage, WeightBuffAccess, TEXT("FModel_TensorWeights"));
+					WeightBuff = CreateRHIBuffer(RHICmdList, TensorData.Num(), WeightBuffUsage, WeightBuffAccess, *Tensor.GetName());
 				
 					FMemory::Memcpy(UploadBuffPtr + UploadOffset, TensorData.GetData(), TensorData.Num());
 					RHICmdList.CopyBufferRegion(WeightBuff, 0, UploadBuff, UploadOffset, TensorData.Num());
@@ -1158,7 +1193,7 @@ void FModel::AddDispatchOps_RenderThread(FRDGBuilder& GraphBuilder)
 
 	for (int32 Idx = 0; Idx < WeightTensorIndices.Num(); ++Idx)
 	{
-		if (ConstantCPUTensorIndices.Find(WeightTensorIndices[Idx]) == -1)
+		if (ConstantCPUTensorIndices.Find(WeightTensorIndices[Idx]) == INDEX_NONE)
 		{
 			++NumWeightTensors;
 		}
diff --git a/Engine/Plugins/Experimental/NNE/Source/NNERuntimeRDG/Private/Dml/NNEDmlOperator.cpp b/Engine/Plugins/Experimental/NNE/Source/NNERuntimeRDG/Private/Dml/NNEDmlOperator.cpp
--- a/Engine/Plugins/Experimental/NNE/Source/NNERuntimeRDG/Private/Dml/NNEDmlOperator.cpp
+++ b/Engine/Plugins/Experimental/NNE/Source/NNERuntimeRDG/Private/Dml/NNEDmlOperator.cpp
@@ -62,7 +62,7 @@ namespace DmlUtil
 		ElemSizeInBytes = Tensor.GetElemByteSize();
 
 		SetShape1D(Tensor.GetShape().GetData()[0], Rank);
-		Update(DmlDataType);
+		Update(DmlDataType, Tensor.HasPreparedData());
 
 		return true;
 	}
diff --git a/Engine/Plugins/Experimental/NNE/Source/NNERuntimeRDG/Private/Dml/NNEDmlOperatorConvolution.cpp b/Engine/Plugins/Experimental/NNE/Source/NNERuntimeRDG/Private/Dml/NNEDmlOperatorConvolution.cpp
--- a/Engine/Plugins/Experimental/NNE/Source/NNERuntimeRDG/Private/Dml/NNEDmlOperatorConvolution.cpp
+++ b/Engine/Plugins/Experimental/NNE/Source/NNERuntimeRDG/Private/Dml/NNEDmlOperatorConvolution.cpp
@@ -302,13 +302,13 @@ public:
 		DmlUtil::FTensorDesc	DmlBiasTensor{};
 		DmlUtil::FTensorDesc	DmlOutputTensor{};
 
-		if (!InitDmlTensorDesc(DmlInputTensor, InputTensor))
+		if (!DmlInputTensor.InitFromTensor(InputTensor, 3))
 		{
 			UE_LOG(LogNNE, Warning, TEXT("Failed to initialize tensor(s) for DML inference"));
 			return false;
 		}
 
-		if (!InitDmlTensorDesc(DmlFilterTensor, FilterTensor))
+		if (!DmlFilterTensor.InitFromTensor(FilterTensor, 3))
 		{
 			UE_LOG(LogNNE, Warning, TEXT("Failed to initialize tensor(s) for DML inference"));
 			return false;
@@ -316,12 +316,12 @@ public:
 
 		if (InputTensors.Num() > 2)
 		{
-			NNECore::Internal::FTensor BiasTensor = InputTensors[2];
+			const NNECore::Internal::FTensor& BiasTensor = InputTensors[2];
+
+			FSmallArray Shape;
 
 			if (BiasTensor.GetShape().Rank() < (int32) Args.NumDimensions)
 			{
-				FSmallArray Shape;
-
 				Shape.Add(1);
 				Shape.Add(BiasTensor.GetShape().GetData()[0]);
 				Shape.Add(1);
@@ -330,18 +330,20 @@ public:
 				{
 					Shape.Add(1);
 				}
-
-				BiasTensor.SetShape(NNECore::FTensorShape::Make(Shape));
+			}
+			else
+			{
+				Shape = BiasTensor.GetShape().GetData();
 			}
 
-			if (!InitDmlTensorDesc(DmlBiasTensor, BiasTensor))
+			if (!DmlBiasTensor.InitFromTensor(BiasTensor, 3, MakeEmptyArrayView<uint32>(), Shape))
 			{
 				UE_LOG(LogNNE, Warning, TEXT("Failed to initialize tensor(s) for DML inference"));
 				return false;
 			}
 		}
 
-		if (!InitDmlTensorDesc(DmlOutputTensor, OutputTensor))
+		if (!DmlOutputTensor.InitFromTensor(OutputTensor, 3))
 		{
 			UE_LOG(LogNNE, Warning, TEXT("Failed to initialize tensor(s) for DML inference"));
 			return false;
