NNE - Add RDG Dml MaxPool
#rnx

diff --git a/Engine/Plugins/Experimental/NNE/Source/NNEQA/Resources/NNEQATestDesc.json b/Engine/Plugins/Experimental/NNE/Source/NNEQA/Resources/NNEQATestDesc.json
--- a/Engine/Plugins/Experimental/NNE/Source/NNEQA/Resources/NNEQATestDesc.json
+++ b/Engine/Plugins/Experimental/NNE/Source/NNEQA/Resources/NNEQATestDesc.json
@@ -436,6 +436,15 @@
 			"targets": [
 				{ "target": "Unsqueeze" }
 			]
+		},
+		{
+			"category": "MaxPool",
+			"targets": [
+				{ "target": "MaxPool" }
+			],
+			"runtimes": [
+				{ "name": "NNERuntimeRDGHlsl", "skip": true, "comment": "Not supported" }
+			]
 		}
 	],
 
@@ -1035,6 +1044,18 @@
 			"datasets": [
 				{ "inputs": [ { "shape": [ 2, 4, 3 ] } ], "weights": [ { "shape": [ 3 ], "source": [ "3", "0", "2" ], "type": "int64" } ], "outputs": [ { "shape": [ 3, 4, 2 ] } ] }
 			]
+		},
+		{
+			"name": "MaxPool.2D.Default",
+			"datasets": [
+				{ "inputs": [ { "shape": [ 1, 3, 32, 32 ] } ], "outputs": [ { "shape": [ 1, 3, 31, 31 ] } ] }
+			]
+		},
+		{
+			"name": "MaxPool.2D.Strides2",
+			"datasets": [
+				{ "inputs": [ { "shape": [ 1, 2, 256, 256 ] } ], "outputs": [ { "shape": [ 1, 2, 128, 128 ] } ] }
+			]
 		}
 	],
 
@@ -1495,6 +1516,26 @@
 					{ "name": "axis", "value": -3, "type": "Int32" }
 				] }
 			]
+		},
+		{
+			"name": "MaxPool.2D.Default",
+			"attribute_maps": [
+				{ "attributes": [ 
+						{ "name": "kernel_shape", "value": [2, 2], "type": "Int32Array" } 
+					]
+				}
+			]
+		},
+		{
+			"name": "MaxPool.2D.Strides2",
+			"attribute_maps": [
+				{ "attributes": [ 
+						{ "name": "kernel_shape", "value": [2, 2], "type": "Int32Array" },
+						{ "name": "pads", "value": [0, 0, 0, 0], "type": "Int32Array" },
+						{ "name": "strides", "value": [2, 2], "type": "Int32Array" } 
+					]
+				}
+			]
 		}
 	]
 
diff --git a/Engine/Plugins/Experimental/NNE/Source/NNERuntimeRDG/Private/Dml/NNEDmlOperator.h b/Engine/Plugins/Experimental/NNE/Source/NNERuntimeRDG/Private/Dml/NNEDmlOperator.h
--- a/Engine/Plugins/Experimental/NNE/Source/NNERuntimeRDG/Private/Dml/NNEDmlOperator.h
+++ b/Engine/Plugins/Experimental/NNE/Source/NNERuntimeRDG/Private/Dml/NNEDmlOperator.h
@@ -36,6 +36,12 @@ inline TArrayView<T> MakeEmptyArrayView()
 	return MakeArrayView(static_cast<T*>(nullptr), 0);
 }
 
+template<typename T>
+inline TConstArrayView<T> MakeEmptyConstArrayView()
+{
+	return TConstArrayView<T>(static_cast<const T*>(nullptr), 0);
+}
+
 //
 //
 //
@@ -56,6 +62,67 @@ namespace DmlUtil
 	using FSmallIntArray = TArray<int32, TInlineAllocator<NNECore::FTensorShape::MaxRank>>;
 	using FSmallUIntArray = TArray<uint32, TInlineAllocator<NNECore::FTensorShape::MaxRank>>;
 
+	template<typename InputType, typename OutputType>
+	inline bool ConvertArrayViewNoOverflow(TConstArrayView<InputType> InputView, TArrayView<OutputType>& OutputView)
+	{
+		OutputView = MakeArrayView((OutputType*) InputView.GetData(), InputView.Num());
+		for(int32 Idx = 0; Idx < InputView.Num(); ++Idx)
+		{
+			if(InputView[Idx] != static_cast<InputType>(OutputView[Idx]))
+			{
+				return false;
+			}
+		}
+		return true;
+	}
+
+	template<typename OutputType, typename AllocatorType>
+	inline bool GetArrayAttributeNoOverflow(
+		const FNNEAttributeValue* Attr, 
+		TArray<OutputType, AllocatorType>& OutputArray, 
+		TConstArrayView<OutputType> DefaultValues = MakeEmptyConstArrayView<OutputType>()
+		)
+	{
+		if (Attr)
+		{
+
+			TArrayView<OutputType> ConvertedView;
+			TArray<int32> IntArray;
+			TArray<float> FloatArray;
+
+			switch(Attr->GetType())
+			{
+			case ENNEAttributeDataType::Int32Array:
+				{
+					IntArray = Attr->GetValue<TArray<int32>>();
+					if(!ConvertArrayViewNoOverflow(TConstArrayView<int32>(IntArray), ConvertedView))
+					{
+						return false;
+					}
+				}
+				break;
+			case ENNEAttributeDataType::FloatArray:
+				{
+					FloatArray = Attr->GetValue<TArray<float>>();
+					if(!ConvertArrayViewNoOverflow(TConstArrayView<float>(FloatArray), ConvertedView))
+					{
+						return false;
+					}
+				}
+				break;
+			default:
+				return false;
+			}
+			
+			OutputArray = TArray<OutputType, AllocatorType>{ ConvertedView };
+		}
+		else
+		{
+			OutputArray = TArray<OutputType, AllocatorType>{ DefaultValues };
+		}
+		return true;
+	}
+
 	struct FTensorDesc
 	{
 		DML_BUFFER_TENSOR_DESC	BuffDesc;
diff --git a/Engine/Plugins/Experimental/NNE/Source/NNERuntimeRDG/Private/Dml/NNEDmlOperatorConvolution.cpp b/Engine/Plugins/Experimental/NNE/Source/NNERuntimeRDG/Private/Dml/NNEDmlOperatorConvolution.cpp
--- a/Engine/Plugins/Experimental/NNE/Source/NNERuntimeRDG/Private/Dml/NNEDmlOperatorConvolution.cpp
+++ b/Engine/Plugins/Experimental/NNE/Source/NNERuntimeRDG/Private/Dml/NNEDmlOperatorConvolution.cpp
@@ -3,6 +3,7 @@
 #ifdef NNE_USE_DIRECTML
 
 #include "NNEDmlOperator.h"
+#include "NNEDmlPaddingCommon.h"
 
 namespace UE::NNERuntimeRDG::Private::Dml
 {
@@ -14,44 +15,12 @@ class FOperatorDmlConv : public FOperatorDml
 {
 	using FSmallArray = TArray<uint32, TInlineAllocator<NcdhwSpatialDimensionCount>>;
 	using FIntArray = TArray<int32>;
-	
-	enum EAutoPad
-	{
-		NOTSET,
-		SAME_UPPER,
-		SAME_LOWER,
-		VALID
-	};
-
-	static EAutoPad AutoPadFromString(FStringView StringVal) 
-	{
-		if (FCString::Stricmp(StringVal.GetData(), TEXT("NOTSET")) == 0) 
-		{
-			return EAutoPad::NOTSET;
-		}
-		else if (FCString::Stricmp(StringVal.GetData(), TEXT("SAME_UPPER")) == 0)
-		{
-			return EAutoPad::SAME_UPPER;
-		}
-		else if (FCString::Stricmp(StringVal.GetData(), TEXT("SAME_LOWER")) == 0)
-		{
-			return EAutoPad::SAME_LOWER;
-		}
-		else if (FCString::Stricmp(StringVal.GetData(), TEXT("VALID")) == 0)
-		{
-			return EAutoPad::VALID;
-		}
-		else
-		{
-			return EAutoPad::NOTSET;
-		}
-	}
 
 	struct FConvArgs
 	{
 		EAutoPad		AutoPad;
-		FSmallArray		StartPadding;
-		FSmallArray		EndPadding;
+		DmlUtil::FSmallUIntArray StartPadding;
+		DmlUtil::FSmallUIntArray EndPadding;
 		FIntArray		OutPadding;
 		FIntArray		Dilations;
 		FIntArray		Strides;
@@ -123,40 +92,13 @@ class FOperatorDmlConv : public FOperatorDml
 
 			Group = Attributes.GetValueOrDefault<int32>(TEXT("group"), 1);
 
-			AutoPad = AutoPadFromString(*Attributes.GetValue<FString>(TEXT("auto_pad")));
-
-			if (AutoPad == EAutoPad::NOTSET)
-			{
-				const FNNEAttributeValue* AttrPads = Attributes.GetAttributeValue(TEXT("pads"));
-				TArray<int32>	Pads;
-
-				if (AttrPads)
-				{
-					Pads = AttrPads->GetValue<TArray<int32>>();
-				}
-				else
-				{
-					Pads.Init(0, 2 * NumDimensions);
-				}
-
-				for (uint32 Dim = 0; Dim < NumDimensions; ++Dim)
-				{
-					StartPadding.Add(Pads[Dim]);
-					EndPadding.Add(Pads[Dim + NumDimensions]);
-				}
-			}
-			else if (AutoPad == EAutoPad::VALID)
-			{
-				StartPadding.Init(0, NumDimensions);
-				EndPadding.Init(0, NumDimensions);
-			}
-			else
-			{
-				StartPadding.Init(0, NumDimensions);
-				EndPadding.Init(0, NumDimensions);
-
-				SetAutoPadding(InputShape.GetData(), FilterShape.GetData());
-			}
+			ComputeStartEndPaddings(
+				InputShape.GetData(),
+				Attributes, 
+				StartPadding, 
+				EndPadding,
+				ConvolutionPadding(InputShape.GetData())
+			);
 
 			if (Direction == DML_CONVOLUTION_DIRECTION_FORWARD)
 			{
@@ -187,38 +129,29 @@ class FOperatorDmlConv : public FOperatorDml
 		//
 		//
 		//
-		void SetAutoPadding(TConstArrayView<uint32> InputShape, TConstArrayView<uint32> FilterShape)
+		DmlUtil::FSmallUIntArray ConvolutionPadding(TConstArrayView<uint32> InputShape)
 		{
 			const uint32 DimOffset = NonspatialDimensionCount;
 
-			for (uint32 Dim = 0; Dim < NumDimensions; ++Dim)
+			if(Direction == DML_CONVOLUTION_DIRECTION_FORWARD)
 			{
-				uint32 Padding;
-
-				if (Direction == DML_CONVOLUTION_DIRECTION_FORWARD)
-				{
-					uint32 InputLen = uint32(InputShape[Dim + DimOffset]);
-					uint32 StridedOutLen = (InputLen + Strides[Dim] - 1) / Strides[Dim];
-					uint32 KernelLen = 1 + (WindowSize[Dim] - 1) * Dilations[Dim];
-					uint32 Len = Strides[Dim] * (StridedOutLen - 1) + KernelLen;
-					
-					Padding = (Len <= InputLen) ? 0 : (Len - InputLen);
-				}
-				else
-				{
-					Padding = (InputShape[Dim + DimOffset] - 1) * Dilations[Dim] - Strides[Dim] + OutPadding[Dim] + 1;
-				}
-
-				if (AutoPad == EAutoPad::SAME_LOWER)
-				{
-					StartPadding[Dim] = (Padding + 1) / 2;
-				}
-				else
+				return KernelPadding(
+					InputShape, WindowSize,
+					MakeArrayView((uint32*) Dilations.GetData(), Dilations.Num()), MakeArrayView((uint32*) Strides.GetData(), Strides.Num())
+				);
+			}
+			// Deconvolution
+			else
+			{
+				DmlUtil::FSmallUIntArray Padding;
+				Padding.SetNumUninitialized(NumDimensions);
+				
+				for (uint32 Dim = 0; Dim < NumDimensions; ++Dim)
 				{
-					StartPadding[Dim] = Padding / 2;
+					Padding[Dim] = (InputShape[Dim + DimOffset] - 1) * Dilations[Dim] - Strides[Dim] + OutPadding[Dim] + 1;
 				}
 
-				EndPadding[Dim] = Padding - StartPadding[Dim];
+				return Padding;
 			}
 		}
 
diff --git a/Engine/Plugins/Experimental/NNE/Source/NNERuntimeRDG/Private/Dml/NNEDmlOperatorMaxPool.cpp b/Engine/Plugins/Experimental/NNE/Source/NNERuntimeRDG/Private/Dml/NNEDmlOperatorMaxPool.cpp
new file mode 100644
--- /dev/null
+++ b/Engine/Plugins/Experimental/NNE/Source/NNERuntimeRDG/Private/Dml/NNEDmlOperatorMaxPool.cpp
@@ -0,0 +1,174 @@
+// Copyright Epic Games, Inc. All Rights Reserved.
+
+#ifdef NNE_USE_DIRECTML
+
+#include "NNEDmlOperator.h"
+#include "NNEDmlPaddingCommon.h"
+
+#include "Algo/AnyOf.h"
+
+namespace UE::NNERuntimeRDG::Private::Dml
+{
+
+/**
+ * MaxPool
+ */
+class FOperatorDmlMaxPool : public FOperatorDml
+{
+    using FIntArray = TArray<int32>;
+
+public:
+
+	static FOperatorDml* Create()
+	{
+		return new FOperatorDmlMaxPool();
+	}
+
+	//
+	//
+	//
+	virtual bool Initialize(IDMLDevice* Device, TArrayView<const NNECore::Internal::FTensor> InputTensors, TArrayView<const NNECore::Internal::FTensor> OutputTensors, const NNECore::FAttributeMap& Attributes) override
+	{
+        //TODO: int64 attributes
+		check(InputTensors.Num() == 1);
+        TConstArrayView<uint32> InputShape = InputTensors[0].GetShape().GetData();
+        check(InputShape.Num() >= 2);
+		check(OutputTensors.Num() == 1 || OutputTensors.Num() == 2);
+
+		const uint32 NumSpatialDimensions = InputShape.Num() - NonspatialDimensionCount;
+        if(NumSpatialDimensions != 2 && NumSpatialDimensions != 3)
+        {
+            UE_LOG(LogNNE, Error, TEXT("Number of spatial dimensions must be in range [2, 3] for DML inference."));
+			return false;
+        }
+
+		DmlUtil::FSmallUIntArray StartPadding, EndPadding, KernelShape, Strides, Dilations;
+
+        Strides.Init(1u, NumSpatialDimensions);
+        if(!DmlUtil::GetArrayAttributeNoOverflow(Attributes.GetAttributeValue(TEXT("strides")), Strides, TConstArrayView<uint32>(Strides)))
+        {
+            UE_LOG(LogNNE, Error, TEXT("Strides attribute cast led to overflow"));
+			return false;
+        }
+		check(Strides.Num() == NumSpatialDimensions);
+
+        Dilations.Init(1u, NumSpatialDimensions);
+        if(!DmlUtil::GetArrayAttributeNoOverflow(Attributes.GetAttributeValue(TEXT("dilations")), Dilations, TConstArrayView<uint32>(Dilations)))
+        {
+            UE_LOG(LogNNE, Error, TEXT("Dilations attribute cast led to overflow"));
+			return false;
+        }
+		check(Dilations.Num() == NumSpatialDimensions);
+
+        if(!DmlUtil::GetArrayAttributeNoOverflow(Attributes.GetAttributeValue(TEXT("kernel_shape")), KernelShape))
+        {
+            UE_LOG(LogNNE, Error, TEXT("kernel_shape attribute cast led to overflow"));
+			return false;
+        }
+        if(KernelShape.IsEmpty())
+        {
+            UE_LOG(LogNNE, Error, TEXT("kernel_shape attribute is required for MaxPool"));
+			return false;
+        }
+		check(KernelShape.Num() == NumSpatialDimensions);
+
+		ComputeStartEndPaddings(
+			InputShape,
+			Attributes, 
+			StartPadding, 
+			EndPadding,
+            KernelPadding(InputShape, KernelShape, Dilations, Strides)
+		);
+
+        for(int Idx = 0; Idx < OutputTensors.Num(); ++Idx)
+        {
+            check(OutputTensors[Idx].GetShape().Rank() == InputShape.Num());
+            check(OutputTensors[Idx].GetShape().GetData()[0] == InputShape[0]);
+            check(OutputTensors[Idx].GetShape().GetData()[1] == InputShape[1]);
+
+            for (uint32 Dim = 0; Dim < NumSpatialDimensions; ++Dim) {
+                uint32 PaddedSize = InputShape[Dim + NonspatialDimensionCount] + StartPadding[Dim] + EndPadding[Dim];
+                check(OutputTensors[Idx].GetShape().GetData()[Dim + NonspatialDimensionCount] == (PaddedSize - KernelShape[Dim]) / Strides[Dim] + 1);
+            }
+        }
+
+        //storage_order is not supported
+        int32 StorageOrder = 
+			Attributes.GetValueOrDefault<int32>(
+				TEXT("storage_order"), 
+				0);
+        if(StorageOrder != 0)
+        {
+            UE_LOG(LogNNE, Error, TEXT("storage_order != 0 is not supported for DML inference"));
+			return false;
+        }
+
+        
+        DmlUtil::FTensorDesc DmlInputTensorDesc;
+        if (!DmlInputTensorDesc.InitFromTensor(InputTensors[0], 4))
+		{
+			UE_LOG(LogNNE, Error, TEXT("Failed to initialize MaxPool's input tensor for DML inference"));
+			return false;
+		}
+
+        DmlUtil::FTensorDesc DmlOutputTensorDesc;
+        if (!DmlOutputTensorDesc.InitFromTensor(OutputTensors[0], 4))
+		{
+			UE_LOG(LogNNE, Error, TEXT("Failed to initialize MaxPool's output tensor for DML inference"));
+			return false;
+		}
+
+        auto FillPoolingDesc = [&](auto& PoolingDesc)
+        {
+            PoolingDesc.InputTensor = &DmlInputTensorDesc.Desc;
+            PoolingDesc.OutputTensor = &DmlOutputTensorDesc.Desc;
+            PoolingDesc.DimensionCount = NumSpatialDimensions;
+            PoolingDesc.WindowSize = KernelShape.GetData();
+            PoolingDesc.Strides = Strides.GetData();
+            PoolingDesc.StartPadding = StartPadding.GetData();
+            PoolingDesc.EndPadding = EndPadding.GetData();
+        };
+
+        const bool bHasDilations =
+            Algo::AnyOf(
+                Dilations,
+                [](auto Dim) {return Dim != 1u; }
+            );
+        const bool bHasOutputIndices = OutputTensors.Num() > 1;
+
+        if (bHasOutputIndices || bHasDilations)
+        {
+            DML_MAX_POOLING2_OPERATOR_DESC DmlMaxPoolOpDesc = {};
+
+            if (bHasOutputIndices)
+            {
+                DmlUtil::FTensorDesc DmlIndicesTensorDesc;
+                if (!DmlIndicesTensorDesc.InitFromTensor(OutputTensors[1], 4))
+                {
+                    UE_LOG(LogNNE, Warning, TEXT("Failed to initialize MaxPool's indices output tensor for DML inference"));
+                    return false;
+                }
+                DmlIndicesTensorDesc.BuffDesc.DataType = DML_TENSOR_DATA_TYPE::DML_TENSOR_DATA_TYPE_UINT64;
+                DmlMaxPoolOpDesc.OutputIndicesTensor = &DmlIndicesTensorDesc.Desc;
+            }
+
+            DmlMaxPoolOpDesc.Dilations = Dilations.GetData();
+            FillPoolingDesc(DmlMaxPoolOpDesc);
+            return CreateOperator(Device, DML_OPERATOR_DESC{ DML_OPERATOR_MAX_POOLING2, &DmlMaxPoolOpDesc} );
+        }
+        else
+        {
+            DML_MAX_POOLING_OPERATOR_DESC DmlMaxPoolOpDesc = {};
+            FillPoolingDesc(DmlMaxPoolOpDesc);
+            return CreateOperator(Device, DML_OPERATOR_DESC{ DML_OPERATOR_MAX_POOLING, &DmlMaxPoolOpDesc} );
+        }
+
+	}
+};
+
+// Register MaxPool operator on Module startup
+NNE_DML_REGISTER_OP(MaxPool)
+
+} // namespace UE::NNERuntimeRDG::Private::Dml
+
+#endif // NNE_USE_DIRECTML
diff --git a/Engine/Plugins/Experimental/NNE/Source/NNERuntimeRDG/Private/Dml/NNEDmlPaddingCommon.h b/Engine/Plugins/Experimental/NNE/Source/NNERuntimeRDG/Private/Dml/NNEDmlPaddingCommon.h
new file mode 100644
--- /dev/null
+++ b/Engine/Plugins/Experimental/NNE/Source/NNERuntimeRDG/Private/Dml/NNEDmlPaddingCommon.h
@@ -0,0 +1,118 @@
+// Copyright Epic Games, Inc. All Rights Reserved.
+
+#pragma once
+
+#include "NNEDmlOperator.h"
+
+namespace UE::NNERuntimeRDG::Private::Dml
+{
+
+enum EAutoPad
+{
+    NOTSET,
+    SAME_UPPER,
+    SAME_LOWER,
+    VALID
+};
+
+static EAutoPad AutoPadFromString(FStringView StringVal) 
+{
+    if (FCString::Stricmp(StringVal.GetData(), TEXT("NOTSET")) == 0) 
+    {
+        return EAutoPad::NOTSET;
+    }
+    else if (FCString::Stricmp(StringVal.GetData(), TEXT("SAME_UPPER")) == 0)
+    {
+        return EAutoPad::SAME_UPPER;
+    }
+    else if (FCString::Stricmp(StringVal.GetData(), TEXT("SAME_LOWER")) == 0)
+    {
+        return EAutoPad::SAME_LOWER;
+    }
+    else if (FCString::Stricmp(StringVal.GetData(), TEXT("VALID")) == 0)
+    {
+        return EAutoPad::VALID;
+    }
+    else
+    {
+        return EAutoPad::NOTSET;
+    }
+}
+
+static DmlUtil::FSmallUIntArray KernelPadding(
+    TConstArrayView<uint32> InputShape, TConstArrayView<uint32> WindowSize, 
+    TConstArrayView<uint32> Dilations, TConstArrayView<uint32> Strides
+    )
+{
+    const uint32 NumSpatialDimensions = InputShape.Num() - NonspatialDimensionCount;
+    check(NumSpatialDimensions >= 1);
+
+    DmlUtil::FSmallUIntArray Padding;
+    Padding.SetNumUninitialized(NumSpatialDimensions);
+
+    for (uint32 Dim = 0; Dim < NumSpatialDimensions; ++Dim)
+    {
+        uint32 InputLen = uint32(InputShape[Dim + NonspatialDimensionCount]);
+        uint32 StridedOutLen = (InputLen + Strides[Dim] - 1) / Strides[Dim];
+        uint32 KernelLen = 1 + (WindowSize[Dim] - 1) * Dilations[Dim];
+        uint32 Len = Strides[Dim] * (StridedOutLen - 1) + KernelLen;
+        
+        Padding[Dim] = (Len <= InputLen) ? 0 : (Len - InputLen);
+    }
+
+    return Padding;
+}
+
+static void ComputeStartEndPaddings(
+    TConstArrayView<uint32> InputShape,
+    const NNECore::FAttributeMap& Attributes, 
+    DmlUtil::FSmallUIntArray &OutStartPadding, 
+    DmlUtil::FSmallUIntArray &OutEndPadding,
+    TConstArrayView<uint32> Padding)
+{
+    const uint32 NumSpatialDimensions = InputShape.Num() - NonspatialDimensionCount;
+    check(NumSpatialDimensions >= 1);
+
+    EAutoPad AutoPad = AutoPadFromString(*Attributes.GetValue<FString>(TEXT("auto_pad")));
+
+    if (AutoPad == EAutoPad::NOTSET)
+    {
+        DmlUtil::FSmallUIntArray Pads;
+
+        Pads.Init(0u, 2 * NumSpatialDimensions);
+        check(DmlUtil::GetArrayAttributeNoOverflow(Attributes.GetAttributeValue(TEXT("pads")), Pads, TConstArrayView<uint32>(Pads)))
+
+
+        for (uint32 Dim = 0; Dim < NumSpatialDimensions; ++Dim)
+        {
+            OutStartPadding.Add(Pads[Dim]);
+            OutEndPadding.Add(Pads[Dim + NumSpatialDimensions]);
+        }
+    }
+    else if (AutoPad == EAutoPad::VALID)
+    {
+        OutStartPadding.Init(0, NumSpatialDimensions);
+        OutEndPadding.Init(0, NumSpatialDimensions);
+    }
+    else
+    {
+        OutStartPadding.Init(0, NumSpatialDimensions);
+        OutEndPadding.Init(0, NumSpatialDimensions);
+
+        for (uint32 Dim = 0; Dim < NumSpatialDimensions; ++Dim)
+        {
+            if (AutoPad == EAutoPad::SAME_LOWER)
+            {
+                OutStartPadding[Dim] = (Padding[Dim] + 1) / 2;
+            }
+            else
+            {
+                OutStartPadding[Dim] = Padding[Dim] / 2;
+            }
+
+            OutEndPadding[Dim] = Padding[Dim] - OutStartPadding[Dim];
+        }
+    }
+}
+
+}
\ No newline at end of file
