NNE - Check int64 to int32 overflow in FModelOptimizerONNXToNNERT::BuildNNEFormat()

#rnx

diff --git a/Engine/Plugins/Experimental/NNE/Source/NNEUtils/Private/NNEUtilsModelOptimizer.cpp b/Engine/Plugins/Experimental/NNE/Source/NNEUtils/Private/NNEUtilsModelOptimizer.cpp
--- a/Engine/Plugins/Experimental/NNE/Source/NNEUtils/Private/NNEUtilsModelOptimizer.cpp
+++ b/Engine/Plugins/Experimental/NNE/Source/NNEUtils/Private/NNEUtilsModelOptimizer.cpp
@@ -574,6 +574,11 @@ private:
 					for (int i = 0; i < AttrValue.count; i++)
 					{
 						Values[i] = (int32)AttrValue.ints[i];
+						if((int64)Values[i] != AttrValue.ints[i])
+						{
+							UE_LOG(LogNNE, Error, TEXT("Overflow detected when converting to int32 attribute '%s' in node '%s'"), ANSI_TO_TCHAR(AttrInfo.name), ANSI_TO_TCHAR(NodeInfo.opName));
+							return false;
+						}
 					}
 
 					Builder->AddOperatorAttribute(Op, AttrInfo.name, FNNEAttributeValue(Values));
