NNE - Add Clip to RDG Dml

#rnx

diff --git a/Engine/Plugins/Experimental/NNE/Source/NNEQA/Resources/NNEQATestDesc.json b/Engine/Plugins/Experimental/NNE/Source/NNEQA/Resources/NNEQATestDesc.json
--- a/Engine/Plugins/Experimental/NNE/Source/NNEQA/Resources/NNEQATestDesc.json
+++ b/Engine/Plugins/Experimental/NNE/Source/NNEQA/Resources/NNEQATestDesc.json
@@ -157,9 +157,7 @@
 				{	"target": "BitShift", "skip": true, "comment": "TODO need attributes" },
 				{	"target": "Cast",     "skip": true, "comment": "TODO need attributes" },
 				{	"target": "Ceil"},
-				{	"target": "Clip", "additional_datasets": [ "Elementwise.Clip.AttributeTest" ], "runtimes": [
-						{ "name": "NNERuntimeRDGDml", "skip": true, "comment": "RDG need scalar tensor support"}
-					]},
+				{	"target": "Clip", "additional_datasets": [ "Elementwise.Clip.AttributeTest" ] },
 				{	"target": "Cos" },
 				{	"target": "Cosh" },
 				{	"target": "Elu", "additional_datasets": [ "Elementwise.Elu.AttributeTest" ], "runtimes": [
@@ -487,7 +485,8 @@
 			"datasets": [
 				{ "inputs": [ { "shape": [1], "comment": "input" }, { "shape": [1], "comment": "min" } ] },
 				{ "inputs": [ { "shape": [1], "comment": "input" }, { "shape": [1], "comment": "min" }, { "shape": [1], "comment": "max" } ] }
-			]
+			],
+			"runtimes": [{ "name": "NNERuntimeRDGDml", "skip": true, "comment": "RDG need scalar tensor support"}]
 		},
 		{
 			"name": "Elementwise.Slow.OneInput",
diff --git a/Engine/Plugins/Experimental/NNE/Source/NNERuntimeRDG/Private/Dml/NNEDmlOperatorElementwise.cpp b/Engine/Plugins/Experimental/NNE/Source/NNERuntimeRDG/Private/Dml/NNEDmlOperatorElementwise.cpp
--- a/Engine/Plugins/Experimental/NNE/Source/NNERuntimeRDG/Private/Dml/NNEDmlOperatorElementwise.cpp
+++ b/Engine/Plugins/Experimental/NNE/Source/NNERuntimeRDG/Private/Dml/NNEDmlOperatorElementwise.cpp
@@ -146,10 +146,12 @@ public:
 
 private:
 
-	FOperatorDmlElementWiseUnary() : Alpha(0.0f), Beta(0.0f), Gamma(0.0f), Num(1) {}
+	FOperatorDmlElementWiseUnary() : Alpha(0.0f), Beta(0.0f), Gamma(0.0f), Min(TNumericLimits<float>::Min()), Max(TNumericLimits<float>::Max()), Num(1) {}
 	float Alpha;
 	float Beta;
 	float Gamma;
+	float Min;
+	float Max;
 	uint32 Num;
 
 public:
@@ -168,6 +170,31 @@ public:
 		Beta = Attributes.GetValueOrDefault(TEXT("beta"), Beta);
 		Gamma = Attributes.GetValueOrDefault(TEXT("gamma"), Gamma);
 
+		if constexpr (std::is_same_v<DmlElementWiseOpDescType, DML_ELEMENT_WISE_CLIP_OPERATOR_DESC>)
+		{
+			const FNNEAttributeValue* MinAttr = Attributes.GetAttributeValue(TEXT("min"));
+			if(MinAttr)
+			{
+				if(MinAttr->GetType() != ENNEAttributeDataType::Float)
+				{
+					UE_LOG(LogNNE, Error, TEXT("Min attribute of clip must be float for DML inference"));
+					return false;
+				}
+				
+				Min = MinAttr->GetValue<float>();
+			}
+			const FNNEAttributeValue* MaxAttr = Attributes.GetAttributeValue(TEXT("max"));
+			if(MaxAttr)
+			{
+				if(MaxAttr->GetType() != ENNEAttributeDataType::Float)
+				{
+					UE_LOG(LogNNE, Error, TEXT("Max attribute of clip must be float for DML inference"));
+					return false;
+				}
+				Max = MaxAttr->GetValue<float>();
+			}
+		}
+
 		// Initialize tensor descriptor (it's same for both input and output)
 		DmlUtil::FTensorDesc	DmlTensorDesc{};
 
@@ -198,6 +225,14 @@ private:
 		Desc.OutputTensor = &TensorDesc.Desc;
 	}
 
+	void InitDmlOpDesc(DML_ELEMENT_WISE_CLIP_OPERATOR_DESC& Desc, DmlUtil::FTensorDesc& TensorDesc)
+	{
+		Desc.InputTensor = &TensorDesc.Desc;
+		Desc.OutputTensor = &TensorDesc.Desc;
+		Desc.Min = Min;
+		Desc.Max = Max;
+	}
+
 	void InitDmlOpDesc(DML_ACTIVATION_SOFTPLUS_OPERATOR_DESC& Desc, DmlUtil::FTensorDesc& TensorDesc)
 	{
 		Desc.InputTensor = &TensorDesc.Desc;
@@ -236,6 +271,10 @@ private:
 	}
 };
 
+template<> FOperatorDmlElementWiseUnary<DML_ELEMENT_WISE_CLIP_OPERATOR_DESC, NNECore::Internal::EElementWiseUnaryOperatorType::Clip>::FOperatorDmlElementWiseUnary()
+{
+}
+
 template<> FOperatorDmlElementWiseUnary<DML_ACTIVATION_SCALED_ELU_OPERATOR_DESC, NNECore::Internal::EElementWiseUnaryOperatorType::Selu>::FOperatorDmlElementWiseUnary()
 	: Alpha(1.67326319217681884765625f), Beta(0.0f), Gamma(1.05070102214813232421875f), Num(1)
 {
@@ -365,6 +404,7 @@ void RegisterElementWiseUnaryOperators()
 	OP(DML_ELEMENT_WISE_ATAN_OPERATOR_DESC, Atan);
 	OP(DML_ELEMENT_WISE_ATANH_OPERATOR_DESC, Atanh);
 	OP(DML_ELEMENT_WISE_CEIL_OPERATOR_DESC, Ceil);
+	OP(DML_ELEMENT_WISE_CLIP_OPERATOR_DESC, Clip);
 	OP(DML_ELEMENT_WISE_COS_OPERATOR_DESC, Cos);
 	OP(DML_ELEMENT_WISE_COSH_OPERATOR_DESC, Cosh);
 	OP(DML_ACTIVATION_ELU_OPERATOR_DESC, Elu);
