NNE HLSL Optim: Matmul style convolution algo

Adding a first version of a conv algo specialized for 4D with group of 1 and constraint on channel and inputs sizes (%32).
Those constaints allow to be rougly 3.5 times faster than the our generic impl on the style transfer demo,
that being said profiling showed that there are good profiling opportunities to be addressed in follow ups.

Tests:
* style transfer demo on win64 and consoles.
* added dedicated unit tests for the algo + all NNE tests green on win64 editor.

#rnx

diff --git a/Engine/Plugins/Experimental/NNE/Source/NNEHlslShaders/Internal/NNEHlslShadersConvCS.h b/Engine/Plugins/Experimental/NNE/Source/NNEHlslShaders/Internal/NNEHlslShadersConvCS.h
--- a/Engine/Plugins/Experimental/NNE/Source/NNEHlslShaders/Internal/NNEHlslShadersConvCS.h
+++ b/Engine/Plugins/Experimental/NNE/Source/NNEHlslShaders/Internal/NNEHlslShadersConvCS.h
@@ -106,6 +106,8 @@ namespace UE::NNEHlslShaders::Internal
 
 		static EConvGroupSize GetMinimalGroupSize(TArrayView<const int32> WShape);
 
+		static TArray<int32> GetPadding(TArrayView<const uint32> XShape, TArrayView<const uint32> WShape, EConvAutoPad AutoPad, TArrayView<const int32> Dilations, TArrayView<const int32> Strides, TArrayView<const int32> Pads);
+
 		static void LexFromString(EConvAutoPad& OutValue, const TCHAR* StringVal);
 	};
 } // UE::NNEHlslShaders::Internal
\ No newline at end of file
diff --git a/Engine/Plugins/Experimental/NNE/Source/NNEHlslShaders/Internal/NNEHlslShadersConvMatmulCS.h b/Engine/Plugins/Experimental/NNE/Source/NNEHlslShaders/Internal/NNEHlslShadersConvMatmulCS.h
new file mode 100644
--- /dev/null
+++ b/Engine/Plugins/Experimental/NNE/Source/NNEHlslShaders/Internal/NNEHlslShadersConvMatmulCS.h
@@ -0,0 +1,38 @@
+// Copyright Epic Games, Inc. All Rights Reserved.
+
+#pragma once
+
+#include "GlobalShader.h"
+#include "ShaderParameterUtils.h"
+#include "RenderGraphUtils.h"
+
+namespace UE::NNEHlslShaders::Internal
+{
+	class NNEHLSLSHADERS_API FConvMatmulCS : public FGlobalShader
+	{
+		DECLARE_GLOBAL_SHADER(FConvMatmulCS);
+		SHADER_USE_PARAMETER_STRUCT(FConvMatmulCS, FGlobalShader)
+
+	public:
+		BEGIN_SHADER_PARAMETER_STRUCT(FParameters, )
+			SHADER_PARAMETER_RDG_BUFFER_SRV(Buffer<float>, Input)
+			SHADER_PARAMETER_RDG_BUFFER_SRV(Buffer<float>, Weight)
+			SHADER_PARAMETER_RDG_BUFFER_SRV(Buffer<float>, Bias)
+			SHADER_PARAMETER_RDG_BUFFER_UAV(RWBuffer<float>, Output)
+			SHADER_PARAMETER(int32, Ci)
+			SHADER_PARAMETER(int32, Hi)
+			SHADER_PARAMETER(int32, Wi)
+			SHADER_PARAMETER(int32, Cw)
+			SHADER_PARAMETER(int32, Hw)
+			SHADER_PARAMETER(int32, Ww)
+			SHADER_PARAMETER(int32, Ho)
+			SHADER_PARAMETER(int32, Wo)
+			SHADER_PARAMETER(int32, StrideH)
+			SHADER_PARAMETER(int32, StrideW)
+			SHADER_PARAMETER(int32, PadLeft)
+			SHADER_PARAMETER(int32, PadTop)
+		END_SHADER_PARAMETER_STRUCT()
+
+		static FIntVector GetGroupCount(TConstArrayView<uint32> OutputShape);
+	};
+} // UE::NNEHlslShaders::Internal
\ No newline at end of file
diff --git a/Engine/Plugins/Experimental/NNE/Source/NNEHlslShaders/Private/NNEHlslShadersConvCS.cpp b/Engine/Plugins/Experimental/NNE/Source/NNEHlslShaders/Private/NNEHlslShadersConvCS.cpp
--- a/Engine/Plugins/Experimental/NNE/Source/NNEHlslShaders/Private/NNEHlslShadersConvCS.cpp
+++ b/Engine/Plugins/Experimental/NNE/Source/NNEHlslShaders/Private/NNEHlslShadersConvCS.cpp
@@ -25,45 +25,6 @@
 		return Result;
 	}
 
-	TArray<int32> GetPadding(TArrayView<const uint32> XShape, TArrayView<const uint32> WShape, EConvAutoPad AutoPad, TArrayView<const int32> Dilations, TArrayView<const int32> Strides, TArrayView<const int32> Pads)
-	{
-		check(XShape.Num() > 2);
-		check(WShape.Num() == XShape.Num());
-		check(Dilations.Num() == 0 || Dilations.Num() == WShape.Num() - 2);
-		check(Strides.Num() == 0 || Strides.Num() == WShape.Num() - 2);
-		check(AutoPad != EConvAutoPad::NOTSET || Pads.Num() == 2 * (WShape.Num() - 2));
-
-		int32 NumDimensions = XShape.Num() - 2;
-		TArray<int32> Result;
-		Result.Init(0, 2 * NumDimensions);
-
-		if (AutoPad == EConvAutoPad::NOTSET)
-		{
-			return TArray<int32>{Pads};
-		}
-		else if (AutoPad == EConvAutoPad::VALID)
-		{
-			return Result;
-		}
-
-		for (int32 DimensionIndex = 0; DimensionIndex < NumDimensions; DimensionIndex++)
-		{
-			int32 DilatedKernelSize = (DimensionIndex < Dilations.Num() ? Dilations[DimensionIndex] : 1) * (WShape[DimensionIndex + 2] - 1) + 1;
-			int32 TotalPad = ((int32)(((XShape[DimensionIndex + 2] + (DimensionIndex < Strides.Num() ? Strides[DimensionIndex] : 1) - 1) / (DimensionIndex < Strides.Num() ? Strides[DimensionIndex] : 1)) - 1)) * (DimensionIndex < Strides.Num() ? Strides[DimensionIndex] : 1) + DilatedKernelSize - XShape[DimensionIndex + 2];
-			if (AutoPad == EConvAutoPad::SAME_LOWER)
-			{
-				Result[DimensionIndex] = (TotalPad + 1) / 2;
-			}	
-			else
-			{
-				Result[DimensionIndex] = TotalPad / 2;
-			}
-			Result[NumDimensions + DimensionIndex] = TotalPad - Result[DimensionIndex];
-		}
-
-		return Result;
-	}
-
 	int32 GetNumThreadsPerGroup(EConvGroupSize GroupSize)
 	{
 		int32 NumThreadsPerGroup = 128;
@@ -109,6 +70,45 @@
 		OutEnvironment.SetDefine(TEXT("MAX_NUM_DIMENSIONS"), FConvConstants::MAX_NUM_DIMENSIONS);
 	}
 
+	TArray<int32> FConvCS::GetPadding(TArrayView<const uint32> XShape, TArrayView<const uint32> WShape, EConvAutoPad AutoPad, TArrayView<const int32> Dilations, TArrayView<const int32> Strides, TArrayView<const int32> Pads)
+	{
+		check(XShape.Num() > 2);
+		check(WShape.Num() == XShape.Num());
+		check(Dilations.Num() == 0 || Dilations.Num() == WShape.Num() - 2);
+		check(Strides.Num() == 0 || Strides.Num() == WShape.Num() - 2);
+		check(AutoPad != EConvAutoPad::NOTSET || Pads.Num() == 2 * (WShape.Num() - 2));
+
+		int32 NumDimensions = XShape.Num() - 2;
+		TArray<int32> Result;
+		Result.Init(0, 2 * NumDimensions);
+
+		if (AutoPad == EConvAutoPad::NOTSET)
+		{
+			return TArray<int32>{Pads};
+		}
+		else if (AutoPad == EConvAutoPad::VALID)
+		{
+			return Result;
+		}
+
+		for (int32 DimensionIndex = 0; DimensionIndex < NumDimensions; DimensionIndex++)
+		{
+			int32 DilatedKernelSize = (DimensionIndex < Dilations.Num() ? Dilations[DimensionIndex] : 1) * (WShape[DimensionIndex + 2] - 1) + 1;
+			int32 TotalPad = ((int32)(((XShape[DimensionIndex + 2] + (DimensionIndex < Strides.Num() ? Strides[DimensionIndex] : 1) - 1) / (DimensionIndex < Strides.Num() ? Strides[DimensionIndex] : 1)) - 1)) * (DimensionIndex < Strides.Num() ? Strides[DimensionIndex] : 1) + DilatedKernelSize - XShape[DimensionIndex + 2];
+			if (AutoPad == EConvAutoPad::SAME_LOWER)
+			{
+				Result[DimensionIndex] = (TotalPad + 1) / 2;
+			}
+			else
+			{
+				Result[DimensionIndex] = TotalPad / 2;
+			}
+			Result[NumDimensions + DimensionIndex] = TotalPad - Result[DimensionIndex];
+		}
+
+		return Result;
+	}
+
 	TArray<int32> FConvCS::GetOutputShape(TArrayView<const uint32> XShape, TArrayView<const uint32> WShape, EConvAutoPad AutoPad, TArrayView<const int32> Dilations, TArrayView<const int32> Strides, TArrayView<const int32> Pads)
 	{
 		check(XShape.Num() > 2);
@@ -117,7 +117,7 @@
 		check(Strides.Num() == 0 || Strides.Num() == WShape.Num() - 2);
 		check(Pads.Num() == 0 || Pads.Num() == 2 * (WShape.Num() - 2));
 
-		TArray<int32> Padding = ConvUtils::GetPadding(XShape, WShape, AutoPad, Dilations, Strides, Pads);
+		TArray<int32> Padding = GetPadding(XShape, WShape, AutoPad, Dilations, Strides, Pads);
 
 		TArray<int32> Result;
 		Result.SetNumUninitialized(XShape.Num());
@@ -151,7 +151,7 @@
 
 		int32 NumDimensions = XShape.Num() - 2;
 
-		TArray<int32> Padding = ConvUtils::GetPadding(XShape, WShape, AutoPad, Dilations, Strides, Pads);
+		TArray<int32> Padding = GetPadding(XShape, WShape, AutoPad, Dilations, Strides, Pads);
 		TArray<int32> GroupShape = GetGroupShape(GroupSize, NumDimensions);
 		TArray<int32> YShape = GetOutputShape(XShape, WShape, AutoPad, Dilations, Strides, Pads);
 		TArray<int32> GridShape = ConvUtils::GetGridShape(YShape, GroupShape);
diff --git a/Engine/Plugins/Experimental/NNE/Source/NNEHlslShaders/Private/NNEHlslShadersConvMatmulCS.cpp b/Engine/Plugins/Experimental/NNE/Source/NNEHlslShaders/Private/NNEHlslShadersConvMatmulCS.cpp
new file mode 100644
--- /dev/null
+++ b/Engine/Plugins/Experimental/NNE/Source/NNEHlslShaders/Private/NNEHlslShadersConvMatmulCS.cpp
@@ -0,0 +1,26 @@
+// Copyright Epic Games, Inc. All Rights Reserved.
+
+#include "NNEHlslShadersConvMatmulCS.h"
+#include "NNECore.h"
+
+namespace UE::NNEHlslShaders::Internal
+{
+	FIntVector FConvMatmulCS::GetGroupCount(TConstArrayView<uint32> OutputShape)
+	{
+		constexpr int32 NUM_PIXEL_PER_THREADGROUP = 32 ;
+		constexpr int32 NUM_CHANNEL_PER_THREADGROUP = 32 ;
+
+		check(OutputShape.Num() == 4);
+
+		int32 NumBatch = OutputShape[0];
+		int32 NumChannel = OutputShape[1];
+		int32 NumPixel = OutputShape[2] * OutputShape[3];
+
+		return FIntVector(
+			FMath::DivideAndRoundUp(NumPixel, NUM_PIXEL_PER_THREADGROUP),
+			FMath::DivideAndRoundUp(NumChannel, NUM_CHANNEL_PER_THREADGROUP),
+			NumBatch);
+	}
+
+	IMPLEMENT_GLOBAL_SHADER(FConvMatmulCS, "/NNE/NNEHlslShadersConvMatmul.usf", "ConvMatmul", SF_Compute);
+} // UE::NNEHlslShaders::Internal
\ No newline at end of file
diff --git a/Engine/Plugins/Experimental/NNE/Source/NNEHlslShaders/Shaders/NNEHlslShadersConvMatmul.usf b/Engine/Plugins/Experimental/NNE/Source/NNEHlslShaders/Shaders/NNEHlslShadersConvMatmul.usf
new file mode 100644
--- /dev/null
+++ b/Engine/Plugins/Experimental/NNE/Source/NNEHlslShaders/Shaders/NNEHlslShadersConvMatmul.usf
@@ -0,0 +1,172 @@
+// Copyright Epic Games, Inc. All Rights Reserved.
+
+#include "/Engine/Public/Platform.ush"
+
+#define WORK_TYPE float
+#define BUFFER_TYPE float
+#define BUFFER_TO_WORK_TYPE(x) x
+#define WORK_TO_BUFFER_TYPE(x) x
+
+#define LDS_IC 32 //Input channel cached in LDS
+#define TG_OC 32 //Number of output channel covered by a thread group
+#define TG_OP 32 //Number of output pixel covered by a thread group
+
+groupshared WORK_TYPE LDSInputs[LDS_IC][TG_OP]; //1024 WORK_TYPE
+groupshared WORK_TYPE LDSWeights[LDS_IC][TG_OC]; //1024 WORK_TYPE
+
+Buffer<BUFFER_TYPE> Input; // Ni x Ci x Hi x Wi
+Buffer<BUFFER_TYPE> Weight; // Cw x Ci x Hw x Ww
+Buffer<BUFFER_TYPE> Bias; // Cw
+RWBuffer<BUFFER_TYPE> Output; // Ni x Cw x Ho x Wo
+
+int Ci;
+int Hi;
+int Wi;
+int Ho;
+int Wo;
+int Cw;
+int Hw;
+int Ww;
+int PadLeft;
+int PadTop;
+int StrideH;
+int StrideW;
+
+//#define DEBUGVAL(val, pos, thread) {if (thread == GroupThreadIdx) Output[pos] = val;}
+
+[numthreads(8, 8, 1)]
+void ConvMatmul(
+	in const uint3 DispatchThreadID : SV_DispatchThreadID,
+	in const uint3 GroupID : SV_GroupID,
+	in const uint3 GroupThreadID : SV_GroupThreadID,
+	in const uint GroupIndex : SV_GroupIndex)
+{
+	
+	const int DispatchThreadOutputPixelOffset = 4 * DispatchThreadID.x;
+	const int DispatchThreadOutputChannelOffset = 4 * DispatchThreadID.y;
+	const int GroupThreadOutputPixelOffset = 4 * GroupThreadID.x;
+	const int GroupThreadOutputChannelOffset = 4 * GroupThreadID.y;
+	const int GroupThreadIdx = GroupIndex;
+	const int Scalar_GroupOutputPixelOffset = TG_OP * GroupID.x;
+	const int Scalar_GroupOutputChannelOffset = TG_OC * GroupID.y;
+	const int Scalar_BatchOutputOffset = Ho * Wo * Cw * GroupID.z;
+	const int Scalar_BatchInputOffset = Hi * Wi * Cw * GroupID.z;
+	const int Scalar_GroupOutputPixelH = (Scalar_GroupOutputPixelOffset) / Wo * StrideH;
+	const int Scalar_GroupOutputPixelWBase = Scalar_GroupOutputPixelOffset % Wo * StrideW;
+	const int Scalar_GroupInputPixelKernelTopLeftH = Scalar_GroupOutputPixelH - PadTop;
+	const int Scalar_GroupInputPixelKernelTopLeftW = Scalar_GroupOutputPixelWBase - PadLeft;
+
+	const int cw = DispatchThreadOutputChannelOffset;
+	const int pi = DispatchThreadOutputPixelOffset;
+
+	WORK_TYPE BiasC0 = BUFFER_TO_WORK_TYPE(Bias[cw+0]);
+	WORK_TYPE BiasC1 = BUFFER_TO_WORK_TYPE(Bias[cw+1]);
+	WORK_TYPE BiasC2 = BUFFER_TO_WORK_TYPE(Bias[cw+2]);
+	WORK_TYPE BiasC3 = BUFFER_TO_WORK_TYPE(Bias[cw+3]);
+	
+	int i;
+	// First index is thread output channel offset (on cw)
+	// 2nd index is thread pixels offset (on pi)
+	WORK_TYPE Values[4][4];
+	
+	UNROLL
+	for (i = 0; i < 4; ++i)
+	{
+		Values[0][i] = BiasC0;
+		Values[1][i] = BiasC1;
+		Values[2][i] = BiasC2;
+		Values[3][i] = BiasC3;
+	}
+	
+	int Scalar_KernelIdx = 0;
+	const int InputChannelOffset = GroupThreadIdx >= 32 ? 16 : 0;
+	const int OutputPixelOffset = GroupThreadIdx % 32;
+	const int OutputChannelOffset = GroupThreadIdx % 32;
+	
+	for (int hw = 0; hw < Hw; ++hw)
+	{
+		const int Scalar_InputPixelH = Scalar_GroupInputPixelKernelTopLeftH + hw;
+		const bool Scalar_IsValidInputH = (Scalar_InputPixelH >= 0) && (Scalar_InputPixelH < Hi);
+		
+		for (int ww = 0; ww < Ww; ++ww)
+		{
+			const int InputPixelW = Scalar_GroupInputPixelKernelTopLeftW + ww + OutputPixelOffset * StrideW;
+			const bool IsValidInput = Scalar_IsValidInputH && (InputPixelW >= 0) && (InputPixelW < Wi);
+			int InputPixelOffset = Scalar_InputPixelH * Wi + InputPixelW;
+			
+			InputPixelOffset = clamp(InputPixelOffset, 0, Hi * Wi);
+			
+			for (int ci = 0; ci < Ci; ci += LDS_IC)
+			{
+				///— DDR → LDS
+				// We need to load 2048 floats
+				// 1024 weights (32 inputs channel, 32 output channels) and
+				// 1024 inputs (32 inputs channel, 32 inputs pixels)
+				// We have 64 threads thus each will read 16 inputs and 16 weights.
+				// the first 32 read the first 16 channels, while the next 32 read the remaining 16.
+
+				const int ReadOffsetI = Scalar_BatchInputOffset + (ci + InputChannelOffset) * Hi * Wi + InputPixelOffset;
+				const int ReadOffsetW = (Scalar_GroupOutputChannelOffset + OutputChannelOffset) * Ci * Hw * Ww + (ci + InputChannelOffset) * Hw * Ww + Scalar_KernelIdx;
+				
+				UNROLL
+				for (i = 0; i < (LDS_IC / 2); ++i)
+				{
+					WORK_TYPE ValueI = BUFFER_TO_WORK_TYPE(Input[Hi * Wi * i + ReadOffsetI]);
+					ValueI = IsValidInput ? ValueI : 0.0f;
+					LDSInputs[InputChannelOffset + i][OutputPixelOffset] = ValueI;
+
+					// Idea: Try to transpose weights to W[Kw, Ci, Cw] to get coalesced reads.
+					WORK_TYPE ValueW = BUFFER_TO_WORK_TYPE(Weight[Hw * Ww * i + ReadOffsetW]);
+					LDSWeights[InputChannelOffset + i][OutputChannelOffset] = ValueW;
+				}
+				
+				GroupMemoryBarrierWithGroupSync();
+				
+				/// LDS to register + inner loop
+				// Loop on cached input channels
+				#define INNERLOOPUNROLLCOUNT 4
+				for (int cachedCiBase = 0; cachedCiBase < LDS_IC; cachedCiBase += INNERLOOPUNROLLCOUNT)
+				{
+					WORK_TYPE RegWeights[INNERLOOPUNROLLCOUNT][4];
+					WORK_TYPE RegInputs[INNERLOOPUNROLLCOUNT][4];
+					
+					UNROLL
+					for (int unrolledIdx = 0; unrolledIdx < INNERLOOPUNROLLCOUNT; ++unrolledIdx)
+					{
+						// 8 load from LDS
+						UNROLL
+						for (i = 0; i < 4; ++i)
+						{
+							RegWeights[unrolledIdx][i] = LDSWeights[cachedCiBase + unrolledIdx][GroupThreadOutputChannelOffset + i];
+							RegInputs[unrolledIdx][i] = LDSInputs[cachedCiBase + unrolledIdx][GroupThreadOutputPixelOffset + i];
+						}
+
+						// Inner loop (16 mads)
+						UNROLL
+						for (i = 0; i < 4; ++i)
+						{
+							Values[i][0] += RegWeights[unrolledIdx][i] * RegInputs[unrolledIdx][0];
+							Values[i][1] += RegWeights[unrolledIdx][i] * RegInputs[unrolledIdx][1];
+							Values[i][2] += RegWeights[unrolledIdx][i] * RegInputs[unrolledIdx][2];
+							Values[i][3] += RegWeights[unrolledIdx][i] * RegInputs[unrolledIdx][3];
+						}
+					}
+				}
+				
+				GroupMemoryBarrierWithGroupSync();
+			}
+			++Scalar_KernelIdx;
+		}
+	}
+	
+	/// register to DDR
+	UNROLL
+	for (i = 0; i < 4; ++i)
+	{
+		const int WriteOffset = Scalar_BatchOutputOffset + (cw + i) * Ho * Wo + pi;
+		Output[WriteOffset + 0] = Values[i][0];
+		Output[WriteOffset + 1] = Values[i][1];
+		Output[WriteOffset + 2] = Values[i][2];
+		Output[WriteOffset + 3] = Values[i][3];
+	}
+}
\ No newline at end of file
diff --git a/Engine/Plugins/Experimental/NNE/Source/NNEQA/Resources/NNEQATestDesc.json b/Engine/Plugins/Experimental/NNE/Source/NNEQA/Resources/NNEQATestDesc.json
--- a/Engine/Plugins/Experimental/NNE/Source/NNEQA/Resources/NNEQATestDesc.json
+++ b/Engine/Plugins/Experimental/NNE/Source/NNEQA/Resources/NNEQATestDesc.json
@@ -285,7 +285,7 @@
 		{
 			"category": "Convolution",
 			"targets": [
-				{ "target": "Conv", "absolute_tolerance": 1e-6 }
+				{ "target": "Conv", "absolute_tolerance": 1e-4, "relative_tolerance": 1e-4 }
 			]
 		},
 		{
@@ -698,13 +698,29 @@
 				{ "inputs": [ { "shape": [ 1, 3, 5, 5 ] }, { "shape": [ 1, 3, 3, 3 ] } ], "outputs": [ { "shape": [ 1, 1, 5, 5 ] } ], "comment": "With channels" },
 				{ "inputs": [ { "shape": [ 1, 1, 5, 5 ] }, { "shape": [ 3, 1, 3, 3 ] } ], "outputs": [ { "shape": [ 1, 3, 5, 5 ] } ], "comment": "With out-channels" },
 				{ "inputs": [ { "shape": [ 1, 1, 5, 5 ] }, { "shape": [ 3, 1, 3, 3 ] }, { "shape": [ 3 ] } ], "outputs": [ { "shape": [ 1, 3, 5, 5 ] } ], "comment": "With bias" },
+				{ "inputs": [ { "shape": [ 1, 64, 72, 128] }, { "shape": [ 64, 64, 3, 3 ] }, { "shape": [ 64 ] } ], "outputs": [ { "shape": [ 1, 64, 72, 128 ] } ] },
 				{ "inputs": [ { "shape": [ 3, 1, 5, 5 ] }, { "shape": [ 1, 1, 3, 3 ] } ], "outputs": [ { "shape": [ 3, 1, 5, 5 ] } ], "comment": "Batches" }
 			]
 		},
 		{ "name": "Convolution.WithoutOrAutoPadding", 			"datasets": [ { "inputs": [ { "shape": [ 1, 1, 5, 5 ] }, { "shape": [ 1, 1, 3, 3 ] } ], "outputs": [ { "shape": [ 1, 1, 3, 3 ] } ] } ] },
-		{ "name": "Convolution.WithStridesAndPadding", 			"datasets": [ { "inputs": [ { "shape": [ 1, 1, 7, 5 ] }, { "shape": [ 1, 1, 3, 3 ] } ], "outputs": [ { "shape": [ 1, 1, 4, 3 ] } ] } ] },
-		{ "name": "Convolution.WithoutPaddingAndStrides", 		"datasets": [ { "inputs": [ { "shape": [ 1, 1, 7, 5 ] }, { "shape": [ 1, 1, 3, 3 ] } ], "outputs": [ { "shape": [ 1, 1, 3, 2 ] } ] } ] },
-		{ "name": "Convolution.AsymmetricPaddingAndStrides", 	"datasets": [ { "inputs": [ { "shape": [ 1, 1, 7, 5 ] }, { "shape": [ 1, 1, 3, 3 ] } ], "outputs": [ { "shape": [ 1, 1, 4, 2 ] } ] } ] },
+		{ "name": "Convolution.WithStridesAndPadding", 			
+			"datasets": [ 
+				{ "inputs": [ { "shape": [ 1, 1, 7, 5 ] }, { "shape": [ 1, 1, 3, 3 ] } ], "outputs": [ { "shape": [ 1, 1, 4, 3 ] } ] },
+				{ "inputs": [ { "shape": [ 1, 32, 1, 64 ] }, { "shape": [ 32, 32, 3, 3 ] }, { "shape": [ 32 ] } ], "outputs": [ { "shape": [ 1, 32, 1, 32 ] } ] }
+			] 
+		},
+		{ "name": "Convolution.WithoutPaddingAndStrides",
+	 		"datasets": [ 
+				{ "inputs": [ { "shape": [ 1, 1, 7, 5 ] }, { "shape": [ 1, 1, 3, 3 ] } ], "outputs": [ { "shape": [ 1, 1, 3, 2 ] } ] },
+				{ "inputs": [ { "shape": [ 1, 32, 1, 64 ] }, { "shape": [ 32, 32, 1, 1 ] }, { "shape": [ 32 ] } ], "outputs": [ { "shape": [ 1, 32, 1, 32 ] } ] }
+			]
+		},
+		{ "name": "Convolution.AsymmetricPaddingAndStrides",
+		 	"datasets": [
+				{ "inputs": [ { "shape": [ 1, 1, 7, 5 ] }, { "shape": [ 1, 1, 3, 3 ] } ], "outputs": [ { "shape": [ 1, 1, 4, 2 ] } ] },
+				{ "inputs": [ { "shape": [ 1, 64, 70, 66] }, { "shape": [ 64, 64, 3, 3 ] }, { "shape": [ 64 ] } ], "outputs": [ { "shape": [ 1, 64, 35, 32 ] } ] }
+			]
+		},
 		{ "name": "Convolution.WithGroups", 					"datasets": [ { "inputs": [ { "shape": [ 1, 3, 5, 5 ] }, { "shape": [ 3, 1, 3, 3 ] } ], "outputs": [ { "shape": [ 1, 3, 5, 5 ] } ] } ] },
 		{ "name": "Convolution.Large1", 						"datasets": [ { "inputs": [ { "shape": [ 1, 1, 17, 17 ] }, { "shape": [ 1, 1, 3, 3 ] } ], "outputs": [ { "shape": [ 1, 1, 17, 17 ] } ] } ] },
 		{ "name": "Convolution.Large2", 						"datasets": [ { "inputs": [ { "shape": [ 1, 1, 37, 37 ] }, { "shape": [ 1, 1, 3, 3 ] } ], "outputs": [ { "shape": [ 1, 1, 17, 17 ] } ] } ] },
diff --git a/Engine/Plugins/Experimental/NNE/Source/NNERuntimeRDG/Private/Hlsl/NNERuntimeRDGConv.cpp b/Engine/Plugins/Experimental/NNE/Source/NNERuntimeRDG/Private/Hlsl/NNERuntimeRDGConv.cpp
--- a/Engine/Plugins/Experimental/NNE/Source/NNERuntimeRDG/Private/Hlsl/NNERuntimeRDGConv.cpp
+++ b/Engine/Plugins/Experimental/NNE/Source/NNERuntimeRDG/Private/Hlsl/NNERuntimeRDGConv.cpp
@@ -1,7 +1,9 @@
 // Copyright Epic Games, Inc. All Rights Reserved.
 
 #include "NNERuntimeRDGConv.h"
+#include "Algo/AnyOf.h"
 #include "NNEHlslShadersConvCS.h"
+#include "NNEHlslShadersConvMatmulCS.h"
 #include "NNECoreAttributeMap.h"
 #include "NNERuntimeRDGHlslHelper.h"
 #include "NNECoreTensor.h"
@@ -9,7 +11,8 @@
 
 namespace UE::NNERuntimeRDG::Private::Hlsl
 {
-	DECLARE_GPU_STAT_NAMED(FNNEOperatorConv, TEXT("NNE.Operator.Hlsl.Conv"));
+	DECLARE_GPU_STAT_NAMED(FNNEOperatorConvDefault, TEXT("NNE.Operator.Hlsl.Conv.Default"));
+	DECLARE_GPU_STAT_NAMED(FNNEOperatorConvMatmul, TEXT("NNE.Operator.Hlsl.Conv.Matmul"));
 
 	using EConvAutoPad = UE::NNEHlslShaders::Internal::EConvAutoPad;
 
@@ -108,36 +111,13 @@ namespace UE::NNERuntimeRDG::Private::Hlsl
 			return true;
 		}
 
-		virtual void Dispatch(FRDGBuilder& GraphBuilder, TConstArrayView<FTensorRDGRef> InputTensors, TConstArrayView<FTensorRDGRef> OutputTensors) override
+		void DispatchConvDefault (FRDGBuilder& GraphBuilder, const FTensorRDG& Input, const FTensorRDG& Weights, const FTensorRDG* Bias, const FTensorRDG& Output)
 		{
 			using namespace UE::NNEHlslShaders::Internal;
 
 			constexpr EConvAlgorithm Algorithm = EConvAlgorithm::SharedMemory;
 			constexpr EConvGroupSize GroupSize = EConvGroupSize::Size256;
-
-			check(InputTensors.Num() >= 2 && InputTensors.Num() <= 3);
-			check(OutputTensors.Num() == 1);
-			check(InputTensors[0] != nullptr);
-			check(InputTensors[1] != nullptr);
-			check(OutputTensors[0] != nullptr);
-			
-			const FTensorRDG& Input = *InputTensors[0];
-			const FTensorRDG& Weights = *InputTensors[1];
-			const FTensorRDG& Output = *OutputTensors[0];
-			const FTensorRDG* Bias = nullptr;
-			bool HasBias = false;
-
-			if (InputTensors.Num() == 3) {
-				HasBias = true;
-				check(InputTensors[2] != nullptr);
-				Bias = InputTensors[2];
-			}
-
-			check(Input.GetShape().Rank() > 2);
-			check(Weights.GetShape().Rank() == Input.GetShape().Rank());
-			check(Output.GetShape().Rank() == Input.GetShape().Rank());
-			check(NumDimensions == (Input.GetShape().Rank() - 2));
-
+			bool HasBias = Bias != nullptr;
 			TArray<int32> OutputShape = FConvCS::GetOutputShape(Input.GetShape().GetData(), Weights.GetShape().GetData(), AutoPad, Dilations, Strides, Pads);
 
 			// Set parameters
@@ -159,17 +139,130 @@ namespace UE::NNERuntimeRDG::Private::Hlsl
 			PermutationVector.Set<FConvCS::FConvHasB>(HasBias);
 			TShaderMapRef<FConvCS> ComputeShader(GetGlobalShaderMap(GMaxRHIFeatureLevel), PermutationVector);
 
-			RDG_EVENT_SCOPE(GraphBuilder, "NNE.Operator.Hlsl.Conv");
-			RDG_GPU_STAT_SCOPE(GraphBuilder, FNNEOperatorConv);
+			RDG_EVENT_SCOPE(GraphBuilder, "NNE.Operator.Hlsl.Conv.Default");
+			RDG_GPU_STAT_SCOPE(GraphBuilder, FNNEOperatorConvDefault);
 
 			FComputeShaderUtils::AddPass(
 				GraphBuilder,
-				RDG_EVENT_NAME("NNE.Operator.Hlsl.Conv.Dispatch"),
+				RDG_EVENT_NAME("NNE.Operator.Hlsl.Conv.Default.Dispatch"),
 				ERDGPassFlags::Compute | ERDGPassFlags::NeverCull,
 				ComputeShader,
 				Params,
 				FConvCS::GetGroupCount(OutputShape, FConvCS::GetGroupShape(GroupSize, NumDimensions)));
 		}
+
+		bool DispatchConvMatmul(FRDGBuilder& GraphBuilder, const FTensorRDG& Input, const FTensorRDG& Weights, const FTensorRDG* Bias, const FTensorRDG& Output)
+		{
+			using namespace UE::NNEHlslShaders::Internal;
+
+			if (Group != 1)
+				return false;
+			if (Input.GetShape().Rank() != 4)
+				return false;
+			if (Output.GetShape().Rank() != 4)
+				return false;
+			if (Weights.GetShape().Rank() != 4)
+				return false;
+			if (Bias == nullptr)
+				return false;
+			if (Bias->GetShape().Rank() != 1)
+				return false;
+
+			const bool bHasDilation = Algo::AnyOf(Dilations, [](auto Dim) {return Dim != 1u; });
+			if (bHasDilation)
+				return false;
+
+			int Ni = Input.GetShape().GetData()[0];
+			int Ci = Input.GetShape().GetData()[1];
+			int Hi = Input.GetShape().GetData()[2];
+			int Wi = Input.GetShape().GetData()[3];
+
+			check(Ni == Output.GetShape().GetData()[0]);
+			int Cw = Output.GetShape().GetData()[1];
+			int Ho = Output.GetShape().GetData()[2];
+			int Wo = Output.GetShape().GetData()[3];
+
+			check(Cw == Weights.GetShape().GetData()[0]);
+			check(Ci == Weights.GetShape().GetData()[1]);
+			int Hw = Weights.GetShape().GetData()[2];
+			int Ww = Weights.GetShape().GetData()[3];
+
+			check(Cw == Bias->GetShape().GetData()[0]);
+
+			if (Ci % 32 != 0)
+				return false;
+			if (Cw % 32 != 0)
+				return false;
+			if (Wo % 32 != 0)
+				return false;
+
+			TArray<int32> Padding = FConvCS::GetPadding(Input.GetShape().GetData(), Weights.GetShape().GetData(), AutoPad, Dilations, Strides, Pads);
+
+			// Set parameters
+			FConvMatmulCS::FParameters* Params = GraphBuilder.AllocParameters<FConvMatmulCS::FParameters>();
+			Params->Input = GraphBuilder.CreateSRV(FRDGBufferSRVDesc(Input.GetBuffer(), PF_R32_FLOAT));
+			Params->Weight= GraphBuilder.CreateSRV(FRDGBufferSRVDesc(Weights.GetBuffer(), PF_R32_FLOAT));
+			Params->Bias = GraphBuilder.CreateSRV(FRDGBufferSRVDesc(Bias->GetBuffer(), PF_R32_FLOAT));
+			Params->Output = GraphBuilder.CreateUAV(FRDGBufferUAVDesc(Output.GetBuffer(), PF_R32_FLOAT));
+			Params->Ci = Ci;
+			Params->Hi = Hi;
+			Params->Wi = Wi;
+			Params->Cw = Cw;
+			Params->Hw = Hw;
+			Params->Ww = Ww;
+			Params->Ho = Ho;
+			Params->Wo = Wo;
+			Params->StrideH = Strides[0];
+			Params->StrideW = Strides[1];
+			Params->PadTop = Padding[0];
+			Params->PadLeft = Padding[1];
+
+			FConvMatmulCS::FPermutationDomain PermutationVector;
+			TShaderMapRef<FConvMatmulCS> ComputeShader(GetGlobalShaderMap(GMaxRHIFeatureLevel), PermutationVector);
+
+			RDG_EVENT_SCOPE(GraphBuilder, "NNE.Operator.Hlsl.Conv.Matmul");
+			RDG_GPU_STAT_SCOPE(GraphBuilder, FNNEOperatorConvMatmul);
+
+			FComputeShaderUtils::AddPass(
+				GraphBuilder,
+				RDG_EVENT_NAME("NNE.Operator.Hlsl.Conv.Matmul.Dispatch"),
+				ERDGPassFlags::Compute | ERDGPassFlags::NeverCull,
+				ComputeShader,
+				Params,
+				FConvMatmulCS::GetGroupCount(Output.GetShape().GetData()));
+
+			return true;
+		}
+
+		virtual void Dispatch(FRDGBuilder& GraphBuilder, TConstArrayView<FTensorRDGRef> InputTensors, TConstArrayView<FTensorRDGRef> OutputTensors) override
+		{
+			check(InputTensors.Num() >= 2 && InputTensors.Num() <= 3);
+			check(OutputTensors.Num() == 1);
+			check(InputTensors[0] != nullptr);
+			check(InputTensors[1] != nullptr);
+			check(OutputTensors[0] != nullptr);
+
+			const FTensorRDG& Input = *InputTensors[0];
+			const FTensorRDG& Weights = *InputTensors[1];
+			const FTensorRDG& Output = *OutputTensors[0];
+			const FTensorRDG* Bias = nullptr;
+
+			if (InputTensors.Num() == 3) {
+				check(InputTensors[2] != nullptr);
+				Bias = InputTensors[2];
+			}
+
+			check(Input.GetShape().Rank() > 2);
+			check(Weights.GetShape().Rank() == Input.GetShape().Rank());
+			check(Output.GetShape().Rank() == Input.GetShape().Rank());
+			check(NumDimensions == (Input.GetShape().Rank() - 2));
+
+			if (DispatchConvMatmul(GraphBuilder, Input, Weights, Bias, Output))
+			{
+				return;
+			}
+			DispatchConvDefault(GraphBuilder, Input, Weights, Bias, Output);
+		}
 	};
 
 	bool ValidateConvOperator(const NNECore::FAttributeMap& AttributeMap, TConstArrayView<ENNETensorDataType> InputTypes, TConstArrayView<NNECore::FSymbolicTensorShape> InputShapes)
