- Add RDG DML normalization ops
- BatchNormalization
- MeanVarianceNormalization
- LpNormalization
- LRN
- Renamed NNEDmlPaddingCommon.h to NNEDmlOperatorUtils.h
- Added Dml axis utils

#rnx

diff --git a/Engine/Plugins/Experimental/NNE/Source/NNERuntimeRDG/Private/Dml/NNEDmlModel.cpp b/Engine/Plugins/Experimental/NNE/Source/NNERuntimeRDG/Private/Dml/NNEDmlModel.cpp
--- a/Engine/Plugins/Experimental/NNE/Source/NNERuntimeRDG/Private/Dml/NNEDmlModel.cpp
+++ b/Engine/Plugins/Experimental/NNE/Source/NNERuntimeRDG/Private/Dml/NNEDmlModel.cpp
@@ -924,6 +924,14 @@ bool FModel::Init(TConstArrayView<uint8> ModelData, FDmlDeviceContext* InDevCtx)
 			return false;
 		}
 
+		// Remap inputs
+		TConstArrayView<int32> OpRemappedInputs = OpDesc.Op->GetRemappedInputs();
+
+		for (int32 InputIdx = 0; InputIdx < OpRemappedInputs.Num(); ++InputIdx)
+		{
+			OpInputIndices[OpDesc.InputStart + InputIdx] = OpRemappedInputs[InputIdx];
+		}
+
 		// Filter out the constant CPU inputs from the graph node inputs
 		TConstArrayView<int32> OpConstantCPUInputs = OpDesc.Op->GetConstantCPUInputs();
 
@@ -1182,9 +1190,9 @@ void FModel::AddDispatchOps_RenderThread(FRDGBuilder& GraphBuilder)
 			RDG_EVENT_NAME("FModel_Dispatch_GetInputBuffer"),
 			Params,
 			TransitionBuffFlags,
-			[this, Params](FRHICommandListImmediate& RHICmdList)
+			[this, Idx, Params](FRHICommandListImmediate& RHICmdList)
 			{
-				InputBuffers.Emplace(Params->Buffer->GetRHI());
+				InputBuffers.Insert(Params->Buffer->GetRHI(), Idx);
 			}
 		);
 	}
@@ -1208,15 +1216,15 @@ void FModel::AddDispatchOps_RenderThread(FRDGBuilder& GraphBuilder)
 			RDG_EVENT_NAME("FModel_Dispatch_GetOutputBuffer"),
 			Params,
 			TransitionBuffFlags,
-			[this, Params](FRHICommandListImmediate& RHICmdList)
+			[this, Idx, Params](FRHICommandListImmediate& RHICmdList)
 			{
-				OutputBuffers.Emplace(Params->Buffer->GetRHI());
+				OutputBuffers.Insert(Params->Buffer->GetRHI(), Idx);
 			}
 		);
 	}
 	
 	GraphBuilder.AddPass(
-		RDG_EVENT_NAME("FModel_Dispatch"),
+		RDG_EVENT_NAME("FDmlModel_Dispatch"),
 		ERDGPassFlags::None | ERDGPassFlags::NeverCull,
 		[this, NumWeightTensors](FRHICommandListImmediate& RHICmdList)
 		{
@@ -1228,7 +1236,7 @@ void FModel::AddDispatchOps_RenderThread(FRDGBuilder& GraphBuilder)
 
 					for (int32 Idx = 0; Idx < NumWeightTensors; ++Idx)
 					{
-						InputBuffers.Emplace(nullptr);
+						InputBuffers.Add(nullptr);
 					}
 
 					for (FRHIBuffer* Buffer : InputBuffers)
diff --git a/Engine/Plugins/Experimental/NNE/Source/NNERuntimeRDG/Private/Dml/NNEDmlOperator.cpp b/Engine/Plugins/Experimental/NNE/Source/NNERuntimeRDG/Private/Dml/NNEDmlOperator.cpp
--- a/Engine/Plugins/Experimental/NNE/Source/NNERuntimeRDG/Private/Dml/NNEDmlOperator.cpp
+++ b/Engine/Plugins/Experimental/NNE/Source/NNERuntimeRDG/Private/Dml/NNEDmlOperator.cpp
@@ -353,6 +353,11 @@ TConstArrayView<int32> FOperatorDml::GetConstantCPUInputs() const
 	return ConstantCPUInputs;
 }
 
+TConstArrayView<int32> FOperatorDml::GetRemappedInputs() const
+{
+	return RemappedInputs;
+}
+
 IDMLOperator* FOperatorDml::GetOperator()
 {
 	return DmlOp;
diff --git a/Engine/Plugins/Experimental/NNE/Source/NNERuntimeRDG/Private/Dml/NNEDmlOperator.h b/Engine/Plugins/Experimental/NNE/Source/NNERuntimeRDG/Private/Dml/NNEDmlOperator.h
--- a/Engine/Plugins/Experimental/NNE/Source/NNERuntimeRDG/Private/Dml/NNEDmlOperator.h
+++ b/Engine/Plugins/Experimental/NNE/Source/NNERuntimeRDG/Private/Dml/NNEDmlOperator.h
@@ -170,6 +170,8 @@ public:
 
 	virtual TConstArrayView<int32> GetConstantCPUInputs() const;
 
+	virtual TConstArrayView<int32> GetRemappedInputs() const;
+
 	IDMLOperator* GetOperator();
 
 protected:
@@ -181,6 +183,7 @@ protected:
 
 	TComPtr<IDMLOperator>		DmlOp;
 	DmlUtil::FSmallIntArray		ConstantCPUInputs;
+	DmlUtil::FSmallIntArray		RemappedInputs;
 };
 
 /**
diff --git a/Engine/Plugins/Experimental/NNE/Source/NNERuntimeRDG/Private/Dml/NNEDmlOperatorBatchNormalization.cpp b/Engine/Plugins/Experimental/NNE/Source/NNERuntimeRDG/Private/Dml/NNEDmlOperatorBatchNormalization.cpp
new file mode 100644
--- /dev/null
+++ b/Engine/Plugins/Experimental/NNE/Source/NNERuntimeRDG/Private/Dml/NNEDmlOperatorBatchNormalization.cpp
@@ -0,0 +1,125 @@
+// Copyright Epic Games, Inc. All Rights Reserved.
+
+#ifdef NNE_USE_DIRECTML
+#include "NNEDmlOperator.h"
+
+namespace UE::NNERuntimeRDG::Private::Dml
+{
+
+class FOperatorDmlBatchNormalization : public FOperatorDml
+{
+	static constexpr float DefaultEpsilon = 1e-5f;
+
+	enum
+	{
+		X,
+		Scale,
+		Bias,
+		Mean,
+		Variance,
+		Count
+	};
+
+public:
+
+	static FOperatorDml* Create()
+	{
+		return new FOperatorDmlBatchNormalization();
+	}
+
+	//
+	//
+	//
+	virtual bool Initialize(IDMLDevice* Device, TArrayView<const NNECore::Internal::FTensor> InputTensors, TArrayView<const NNECore::Internal::FTensor> OutputTensors, const NNECore::FAttributeMap& Attributes) override
+	{
+		check(InputTensors.Num() == Count);
+		check(OutputTensors.Num() >= 1);
+
+		RemappedInputs.Add(X);
+		RemappedInputs.Add(Mean);
+		RemappedInputs.Add(Variance);
+		RemappedInputs.Add(Scale);
+		RemappedInputs.Add(Bias);
+
+		const NNECore::Internal::FTensor& InputTensor = InputTensors[0];
+		const NNECore::Internal::FTensor& OutputTensor = OutputTensors[0];
+
+		if (InputTensor.GetShape().Rank() > 8)
+		{
+			UE_LOG(LogNNE, Warning, TEXT("InputTensor rank should be between 1 and 8, got:%d"), InputTensor.GetShape().Rank());
+			return false;
+		}
+
+		// Read attributes
+		const float Epsilon = Attributes.GetValueOrDefault<float>(TEXT("epsilon"), DefaultEpsilon);
+		const int32 bTrainingMode = Attributes.GetValueOrDefault<int32>(TEXT("training_mode"), 0);
+
+		if (bTrainingMode || OutputTensors.Num() > 1)
+		{
+			UE_LOG(LogNNE, Warning, TEXT("DML:BatchNormalization doesn't support training mode"));
+			return false;
+		}
+
+		DmlUtil::FTensorDesc	DmlInputTensors[Count];
+			
+		if (!DmlInputTensors[X].InitFromTensor(InputTensor, InputTensor.GetShape().Rank()))
+		{
+			UE_LOG(LogNNE, Warning, TEXT("Failed to initialize tensor(s) for DML inference"));
+			return false;
+		}
+
+		for (int32 Idx = Scale; Idx < Count; ++Idx)
+		{
+			const NNECore::Internal::FTensor& CurrTensor = InputTensors[Idx];
+			DmlUtil::FTensorDesc& DmlTensor = DmlInputTensors[Idx];
+
+			DmlTensor = DmlUtil::FTensorDesc{};
+
+			if (CurrTensor.GetShape().Rank() == 1)
+			{
+				if (!DmlTensor.InitFromTensor1D(CurrTensor, DmlInputTensors[X].Sizes.Num()))
+				{
+					UE_LOG(LogNNE, Warning, TEXT("Failed to initialize tensor(s) for DML inference"));
+					return false;
+				}
+			}
+			else
+			{
+				if (!DmlTensor.InitFromTensor(CurrTensor, DmlInputTensors[X].Sizes.Num()))
+				{
+					UE_LOG(LogNNE, Warning, TEXT("Failed to initialize tensor(s) for DML inference"));
+					return false;
+				}
+			}
+		}
+		
+		DmlUtil::FTensorDesc	DmlOutputTensor{};
+
+		if (!DmlOutputTensor.InitFromTensor(OutputTensor, DmlInputTensors[X].Sizes.Num()))
+		{
+			UE_LOG(LogNNE, Warning, TEXT("Failed to initialize tensor(s) for DML inference"));
+			return false;
+		}
+
+		DML_BATCH_NORMALIZATION_OPERATOR_DESC	OpDesc{};
+
+		OpDesc.InputTensor = &DmlInputTensors[X].Desc;
+		OpDesc.MeanTensor = &DmlInputTensors[Mean].Desc;
+		OpDesc.VarianceTensor = &DmlInputTensors[Variance].Desc;
+		OpDesc.ScaleTensor = &DmlInputTensors[Scale].Desc;
+		OpDesc.BiasTensor = &DmlInputTensors[Bias].Desc;
+		OpDesc.OutputTensor = &DmlOutputTensor.Desc;
+		OpDesc.Spatial = static_cast<BOOL>(1);
+		OpDesc.Epsilon = Epsilon;
+		OpDesc.FusedActivation = nullptr;
+
+		return CreateOperator(Device, DML_OPERATOR_DESC{ DML_OPERATOR_BATCH_NORMALIZATION, &OpDesc });
+	}
+};
+
+// Register operator on Module startup
+NNE_DML_REGISTER_OP(BatchNormalization)
+
+} // namespace UE::NNERuntimeRDG::Private::Dml
+
+#endif // NNE_USE_DIRECTML
diff --git a/Engine/Plugins/Experimental/NNE/Source/NNERuntimeRDG/Private/Dml/NNEDmlOperatorConvolution.cpp b/Engine/Plugins/Experimental/NNE/Source/NNERuntimeRDG/Private/Dml/NNEDmlOperatorConvolution.cpp
--- a/Engine/Plugins/Experimental/NNE/Source/NNERuntimeRDG/Private/Dml/NNEDmlOperatorConvolution.cpp
+++ b/Engine/Plugins/Experimental/NNE/Source/NNERuntimeRDG/Private/Dml/NNEDmlOperatorConvolution.cpp
@@ -3,7 +3,7 @@
 #ifdef NNE_USE_DIRECTML
 
 #include "NNEDmlOperator.h"
-#include "NNEDmlPaddingCommon.h"
+#include "NNEDmlOperatorUtils.h"
 
 namespace UE::NNERuntimeRDG::Private::Dml
 {
diff --git a/Engine/Plugins/Experimental/NNE/Source/NNERuntimeRDG/Private/Dml/NNEDmlOperatorGather.cpp b/Engine/Plugins/Experimental/NNE/Source/NNERuntimeRDG/Private/Dml/NNEDmlOperatorGather.cpp
--- a/Engine/Plugins/Experimental/NNE/Source/NNERuntimeRDG/Private/Dml/NNEDmlOperatorGather.cpp
+++ b/Engine/Plugins/Experimental/NNE/Source/NNERuntimeRDG/Private/Dml/NNEDmlOperatorGather.cpp
@@ -2,21 +2,11 @@
 
 #ifdef NNE_USE_DIRECTML
 #include "NNEDmlOperator.h"
+#include "NNEDmlOperatorUtils.h"
 
 namespace UE::NNERuntimeRDG::Private::Dml
 {
 
-inline int32 HandleNegativeAxis(int32 Axis, int32 Rank)
-{
-	if (Axis < 0)
-	{
-		Axis += Rank;
-		check(Axis < Rank);
-	}
-
-	return Axis;
-}
-
 class FOperatorDmlGather : public FOperatorDml
 {
 public:
diff --git a/Engine/Plugins/Experimental/NNE/Source/NNERuntimeRDG/Private/Dml/NNEDmlOperatorLRN.cpp b/Engine/Plugins/Experimental/NNE/Source/NNERuntimeRDG/Private/Dml/NNEDmlOperatorLRN.cpp
new file mode 100644
--- /dev/null
+++ b/Engine/Plugins/Experimental/NNE/Source/NNERuntimeRDG/Private/Dml/NNEDmlOperatorLRN.cpp
@@ -0,0 +1,80 @@
+// Copyright Epic Games, Inc. All Rights Reserved.
+
+#ifdef NNE_USE_DIRECTML
+#include "NNEDmlOperator.h"
+#include "NNEDmlOperatorUtils.h"
+
+namespace UE::NNERuntimeRDG::Private::Dml
+{
+
+//
+//
+//
+class FOperatorDmlLRN : public FOperatorDml
+{
+public:
+
+	static FOperatorDml* Create()
+	{
+		return new FOperatorDmlLRN();
+	}
+
+	//
+	//
+	//
+	virtual bool Initialize(IDMLDevice* Device, TArrayView<const NNECore::Internal::FTensor> InputTensors, TArrayView<const NNECore::Internal::FTensor> OutputTensors, const NNECore::FAttributeMap& Attributes) override
+	{
+		check(InputTensors.Num() == 1);
+		check(OutputTensors.Num() == 1);
+
+		const NNECore::Internal::FTensor& InputTensor = InputTensors[0];
+		const NNECore::Internal::FTensor& OutputTensor = OutputTensors[0];
+
+		if (InputTensor.GetShape().Rank() > 8)
+		{
+			UE_LOG(LogNNE, Warning, TEXT("InputTensor rank should be between 1 and 8, got:%d"), InputTensor.GetShape().Rank());
+			return false;
+		}
+
+		// Read attributes
+		const int32 Size = Attributes.GetValueOrDefault<int32>(TEXT("size"), 0);
+		const float Alpha = Attributes.GetValueOrDefault<float>(TEXT("alpha"), 0.0f);
+		const float Beta = Attributes.GetValueOrDefault<float>(TEXT("beta"), 0.0f);
+		const float Bias = Attributes.GetValueOrDefault<float>(TEXT("bias"), 0.0f);
+
+		DmlUtil::FTensorDesc	DmlInputTensor{};
+
+		if (!DmlInputTensor.InitFromTensor(InputTensor, InputTensor.GetShape().Rank()))
+		{
+			UE_LOG(LogNNE, Warning, TEXT("Failed to initialize tensor(s) for DML inference"));
+			return false;
+		}
+		
+		DmlUtil::FTensorDesc	DmlOutputTensor{};
+
+		if (!DmlOutputTensor.InitFromTensor(OutputTensor, DmlInputTensor.Sizes.Num()))
+		{
+			UE_LOG(LogNNE, Warning, TEXT("Failed to initialize tensor(s) for DML inference"));
+			return false;
+		}
+
+		DML_LOCAL_RESPONSE_NORMALIZATION_OPERATOR_DESC	OpDesc{};
+
+		OpDesc.InputTensor = &DmlInputTensor.Desc;
+		OpDesc.OutputTensor = &DmlOutputTensor.Desc;
+		OpDesc.CrossChannel = true; // ONNX only supports cross-channel
+		OpDesc.LocalSize = Size;
+		OpDesc.Alpha = Alpha;
+		OpDesc.Beta = Beta;
+		OpDesc.Bias = Bias;
+
+		return CreateOperator(Device, DML_OPERATOR_DESC{ DML_OPERATOR_LOCAL_RESPONSE_NORMALIZATION, &OpDesc });
+	}
+};
+
+// Register operator on Module startup
+NNE_DML_REGISTER_OP(LRN)
+
+} // namespace UE::NNERuntimeRDG::Private::Dml
+
+#endif // NNE_USE_DIRECTML
diff --git a/Engine/Plugins/Experimental/NNE/Source/NNERuntimeRDG/Private/Dml/NNEDmlOperatorLpNormalization.cpp b/Engine/Plugins/Experimental/NNE/Source/NNERuntimeRDG/Private/Dml/NNEDmlOperatorLpNormalization.cpp
new file mode 100644
--- /dev/null
+++ b/Engine/Plugins/Experimental/NNE/Source/NNERuntimeRDG/Private/Dml/NNEDmlOperatorLpNormalization.cpp
@@ -0,0 +1,82 @@
+// Copyright Epic Games, Inc. All Rights Reserved.
+
+#ifdef NNE_USE_DIRECTML
+#include "NNEDmlOperator.h"
+#include "NNEDmlOperatorUtils.h"
+
+namespace UE::NNERuntimeRDG::Private::Dml
+{
+
+//
+//
+//
+class FOperatorDmlLpNormalization : public FOperatorDml
+{
+	static constexpr float DefaultEpsilon = 1e-5f;
+
+public:
+
+	static FOperatorDml* Create()
+	{
+		return new FOperatorDmlLpNormalization();
+	}
+
+	//
+	//
+	//
+	virtual bool Initialize(IDMLDevice* Device, TArrayView<const NNECore::Internal::FTensor> InputTensors, TArrayView<const NNECore::Internal::FTensor> OutputTensors, const NNECore::FAttributeMap& Attributes) override
+	{
+		check(InputTensors.Num() == 1);
+		check(OutputTensors.Num() == 1);
+
+		const NNECore::Internal::FTensor& InputTensor = InputTensors[0];
+		const NNECore::Internal::FTensor& OutputTensor = OutputTensors[0];
+
+		if (InputTensor.GetShape().Rank() > 8)
+		{
+			UE_LOG(LogNNE, Warning, TEXT("InputTensor rank should be between 1 and 8, got:%d"), InputTensor.GetShape().Rank());
+			return false;
+		}
+
+		// Read attributes
+		const int32 P = Attributes.GetValueOrDefault<int32>(TEXT("p"), 1);
+		check(P >= 1 && P <= 2);
+
+		int32 OnnxAxis = Attributes.GetValueOrDefault<int32>(TEXT("axis"), 0);
+
+		DmlUtil::FTensorDesc	DmlInputTensor{};
+
+		if (!DmlInputTensor.InitFromTensor(InputTensor, InputTensor.GetShape().Rank()))
+		{
+			UE_LOG(LogNNE, Warning, TEXT("Failed to initialize tensor(s) for DML inference"));
+			return false;
+		}
+		
+		DmlUtil::FTensorDesc	DmlOutputTensor{};
+
+		if (!DmlOutputTensor.InitFromTensor(OutputTensor, DmlInputTensor.Sizes.Num()))
+		{
+			UE_LOG(LogNNE, Warning, TEXT("Failed to initialize tensor(s) for DML inference"));
+			return false;
+		}
+
+		uint32 DmlAxis = GetDmlAxis(OnnxAxis, InputTensor.GetShape().Rank(), DmlInputTensor.Sizes.Num());
+
+		DML_LP_NORMALIZATION_OPERATOR_DESC	OpDesc{};
+
+		OpDesc.InputTensor = &DmlInputTensor.Desc;
+		OpDesc.OutputTensor = &DmlOutputTensor.Desc;
+		OpDesc.Axis = DmlAxis;
+		OpDesc.Epsilon = DefaultEpsilon;
+		OpDesc.P = P;
+
+		return CreateOperator(Device, DML_OPERATOR_DESC{ DML_OPERATOR_LP_NORMALIZATION, &OpDesc });
+	}
+};
+
+// Register operator on Module startup
+NNE_DML_REGISTER_OP(LpNormalization)
+
+} // namespace UE::NNERuntimeRDG::Private::Dml
+
+#endif // NNE_USE_DIRECTML
diff --git a/Engine/Plugins/Experimental/NNE/Source/NNERuntimeRDG/Private/Dml/NNEDmlOperatorMaxPool.cpp b/Engine/Plugins/Experimental/NNE/Source/NNERuntimeRDG/Private/Dml/NNEDmlOperatorMaxPool.cpp
--- a/Engine/Plugins/Experimental/NNE/Source/NNERuntimeRDG/Private/Dml/NNEDmlOperatorMaxPool.cpp
+++ b/Engine/Plugins/Experimental/NNE/Source/NNERuntimeRDG/Private/Dml/NNEDmlOperatorMaxPool.cpp
@@ -3,7 +3,7 @@
 #ifdef NNE_USE_DIRECTML
 
 #include "NNEDmlOperator.h"
-#include "NNEDmlPaddingCommon.h"
+#include "NNEDmlOperatorUtils.h"
 
 #include "Algo/AnyOf.h"
 
diff --git a/Engine/Plugins/Experimental/NNE/Source/NNERuntimeRDG/Private/Dml/NNEDmlOperatorMeanVarianceNormalization.cpp b/Engine/Plugins/Experimental/NNE/Source/NNERuntimeRDG/Private/Dml/NNEDmlOperatorMeanVarianceNormalization.cpp
new file mode 100644
--- /dev/null
+++ b/Engine/Plugins/Experimental/NNE/Source/NNERuntimeRDG/Private/Dml/NNEDmlOperatorMeanVarianceNormalization.cpp
@@ -0,0 +1,108 @@
+// Copyright Epic Games, Inc. All Rights Reserved.
+
+#ifdef NNE_USE_DIRECTML
+#include "NNEDmlOperator.h"
+#include "NNEDmlOperatorUtils.h"
+
+namespace UE::NNERuntimeRDG::Private::Dml
+{
+
+//
+//
+//
+class FOperatorDmlMeanVarianceNormalization : public FOperatorDml
+{
+	static constexpr float DefaultEpsilon = 1e-5f;
+
+public:
+
+	static FOperatorDml* Create()
+	{
+		return new FOperatorDmlMeanVarianceNormalization();
+	}
+
+	//
+	//
+	//
+	virtual bool Initialize(IDMLDevice* Device, TArrayView<const NNECore::Internal::FTensor> InputTensors, TArrayView<const NNECore::Internal::FTensor> OutputTensors, const NNECore::FAttributeMap& Attributes) override
+	{
+		check(InputTensors.Num() == 1);
+		check(OutputTensors.Num() == 1);
+
+		const NNECore::Internal::FTensor& InputTensor = InputTensors[0];
+		const NNECore::Internal::FTensor& OutputTensor = OutputTensors[0];
+
+		if (InputTensor.GetShape().Rank() > 8)
+		{
+			UE_LOG(LogNNE, Warning, TEXT("InputTensor rank should be between 1 and 8, got:%d"), InputTensor.GetShape().Rank());
+			return false;
+		}
+
+		// Read attributes
+		const int32 AcrossChannels = Attributes.GetValueOrDefault<int32>(TEXT("across_channels"), 0);
+		const int32 NormalizeVariance = Attributes.GetValueOrDefault<int32>(TEXT("normalize_variance"), 1);
+
+		TArray<int32>	OnnxAxes;
+
+		const FNNEAttributeValue* AttrAxes = Attributes.GetAttributeValue("axes");
+
+		if (AttrAxes)
+		{
+			OnnxAxes = AttrAxes->GetValue<TArray<int32>>();
+		}
+		else
+		{
+			constexpr int32 CrossChannelAxes[] = { 0, 1, 2, 3};
+			constexpr int32 NonChannelAxes[] = { 0, 2, 3 };
+
+			if (AcrossChannels)
+			{
+				OnnxAxes.Append(CrossChannelAxes, UE_ARRAY_COUNT(CrossChannelAxes));
+			}
+			else
+			{
+				OnnxAxes.Append(NonChannelAxes, UE_ARRAY_COUNT(NonChannelAxes));
+			}
+		}
+
+		DmlUtil::FTensorDesc	DmlInputTensor{};
+
+		if (!DmlInputTensor.InitFromTensor(InputTensor, InputTensor.GetShape().Rank()))
+		{
+			UE_LOG(LogNNE, Warning, TEXT("Failed to initialize tensor(s) for DML inference"));
+			return false;
+		}
+		
+		DmlUtil::FTensorDesc	DmlOutputTensor{};
+
+		if (!DmlOutputTensor.InitFromTensor(OutputTensor, DmlInputTensor.Sizes.Num()))
+		{
+			UE_LOG(LogNNE, Warning, TEXT("Failed to initialize tensor(s) for DML inference"));
+			return false;
+		}
+
+		DmlUtil::FSmallUIntArray DmlAxes;
+		SetDmlAxesFromOnnx(DmlAxes, DmlInputTensor.Sizes.Num(), OnnxAxes);
+
+		DML_MEAN_VARIANCE_NORMALIZATION1_OPERATOR_DESC	OpDesc{};
+
+		OpDesc.InputTensor = &DmlInputTensor.Desc;
+		OpDesc.ScaleTensor = nullptr;
+		OpDesc.BiasTensor = nullptr;
+		OpDesc.OutputTensor = &DmlOutputTensor.Desc;
+		OpDesc.AxisCount = DmlAxes.Num();
+		OpDesc.Axes = DmlAxes.GetData();
+		OpDesc.NormalizeVariance = NormalizeVariance;
+		OpDesc.Epsilon = DefaultEpsilon;
+		OpDesc.FusedActivation = nullptr;
+
+		return CreateOperator(Device, DML_OPERATOR_DESC{ DML_OPERATOR_MEAN_VARIANCE_NORMALIZATION1, &OpDesc });
+	}
+};
+
+// Register operator on Module startup
+NNE_DML_REGISTER_OP(MeanVarianceNormalization)
+
+} // namespace UE::NNERuntimeRDG::Private::Dml
+
+#endif // NNE_USE_DIRECTML
diff --git a/Engine/Plugins/Experimental/NNE/Source/NNERuntimeRDG/Private/Dml/NNEDmlOperatorUtils.h b/Engine/Plugins/Experimental/NNE/Source/NNERuntimeRDG/Private/Dml/NNEDmlOperatorUtils.h
new file mode 100644
--- /dev/null
+++ b/Engine/Plugins/Experimental/NNE/Source/NNERuntimeRDG/Private/Dml/NNEDmlOperatorUtils.h
@@ -0,0 +1,149 @@
+// Copyright Epic Games, Inc. All Rights Reserved.
+
+#pragma once
+
+#include "NNEDmlOperator.h"
+
+namespace UE::NNERuntimeRDG::Private::Dml
+{
+
+enum EAutoPad
+{
+    NOTSET,
+    SAME_UPPER,
+    SAME_LOWER,
+    VALID
+};
+
+static EAutoPad AutoPadFromString(FStringView StringVal) 
+{
+    if (FCString::Stricmp(StringVal.GetData(), TEXT("NOTSET")) == 0) 
+    {
+        return EAutoPad::NOTSET;
+    }
+    else if (FCString::Stricmp(StringVal.GetData(), TEXT("SAME_UPPER")) == 0)
+    {
+        return EAutoPad::SAME_UPPER;
+    }
+    else if (FCString::Stricmp(StringVal.GetData(), TEXT("SAME_LOWER")) == 0)
+    {
+        return EAutoPad::SAME_LOWER;
+    }
+    else if (FCString::Stricmp(StringVal.GetData(), TEXT("VALID")) == 0)
+    {
+        return EAutoPad::VALID;
+    }
+    else
+    {
+        return EAutoPad::NOTSET;
+    }
+}
+
+static DmlUtil::FSmallUIntArray KernelPadding(
+    TConstArrayView<uint32> InputShape, TConstArrayView<uint32> WindowSize, 
+    TConstArrayView<uint32> Dilations, TConstArrayView<uint32> Strides
+    )
+{
+    const uint32 NumSpatialDimensions = InputShape.Num() - NonspatialDimensionCount;
+    check(NumSpatialDimensions >= 1);
+
+    DmlUtil::FSmallUIntArray Padding;
+    Padding.SetNumUninitialized(NumSpatialDimensions);
+
+    for (uint32 Dim = 0; Dim < NumSpatialDimensions; ++Dim)
+    {
+        uint32 InputLen = uint32(InputShape[Dim + NonspatialDimensionCount]);
+        uint32 StridedOutLen = (InputLen + Strides[Dim] - 1) / Strides[Dim];
+        uint32 KernelLen = 1 + (WindowSize[Dim] - 1) * Dilations[Dim];
+        uint32 Len = Strides[Dim] * (StridedOutLen - 1) + KernelLen;
+        
+        Padding[Dim] = (Len <= InputLen) ? 0 : (Len - InputLen);
+    }
+
+    return Padding;
+}
+
+static void ComputeStartEndPaddings(
+    TConstArrayView<uint32> InputShape,
+    const NNECore::FAttributeMap& Attributes, 
+    DmlUtil::FSmallUIntArray &OutStartPadding, 
+    DmlUtil::FSmallUIntArray &OutEndPadding,
+    TConstArrayView<uint32> Padding)
+{
+    const uint32 NumSpatialDimensions = InputShape.Num() - NonspatialDimensionCount;
+    check(NumSpatialDimensions >= 1);
+
+    EAutoPad AutoPad = AutoPadFromString(*Attributes.GetValue<FString>(TEXT("auto_pad")));
+
+    if (AutoPad == EAutoPad::NOTSET)
+    {
+        DmlUtil::FSmallUIntArray Pads;
+
+        Pads.Init(0u, 2 * NumSpatialDimensions);
+        check(DmlUtil::GetArrayAttributeNoOverflow(Attributes.GetAttributeValue(TEXT("pads")), Pads, TConstArrayView<uint32>(Pads)))
+
+
+        for (uint32 Dim = 0; Dim < NumSpatialDimensions; ++Dim)
+        {
+            OutStartPadding.Add(Pads[Dim]);
+            OutEndPadding.Add(Pads[Dim + NumSpatialDimensions]);
+        }
+    }
+    else if (AutoPad == EAutoPad::VALID)
+    {
+        OutStartPadding.Init(0, NumSpatialDimensions);
+        OutEndPadding.Init(0, NumSpatialDimensions);
+    }
+    else
+    {
+        OutStartPadding.Init(0, NumSpatialDimensions);
+        OutEndPadding.Init(0, NumSpatialDimensions);
+
+        for (uint32 Dim = 0; Dim < NumSpatialDimensions; ++Dim)
+        {
+            if (AutoPad == EAutoPad::SAME_LOWER)
+            {
+                OutStartPadding[Dim] = (Padding[Dim] + 1) / 2;
+            }
+            else
+            {
+                OutStartPadding[Dim] = Padding[Dim] / 2;
+            }
+
+            OutEndPadding[Dim] = Padding[Dim] - OutStartPadding[Dim];
+        }
+    }
+}
+
+inline int32 HandleNegativeAxis(int32 Axis, int32 Rank)
+{
+	if (Axis < 0)
+	{
+		Axis += Rank;
+		check(Axis < Rank);
+	}
+
+	return Axis;
+}
+
+inline int32 GetDmlAxis(int32 OnnxAxis, int32 OnnxDim, int32 DmlDim)
+{
+	check(DmlDim >= OnnxDim);
+	OnnxAxis = HandleNegativeAxis(OnnxAxis, OnnxDim);
+	uint32 DmlAxis = OnnxAxis + DmlDim - OnnxDim;
+
+	return DmlAxis;
+}
+
+inline void SetDmlAxesFromOnnx(DmlUtil::FSmallUIntArray& DmlAxes, int32 Rank, TConstArrayView<int32> OnnxAxes)
+{
+	DmlAxes.Reset();
+	DmlAxes.Reserve(Rank);
+
+	for (int32 Axis : OnnxAxes)
+	{
+		DmlAxes.Add(GetDmlAxis(Axis, OnnxAxes.Num(), Rank));
+	}
+}
+
+}
\ No newline at end of file
diff --git a/Engine/Plugins/Experimental/NNE/Source/NNERuntimeRDG/Private/Dml/NNEDmlPaddingCommon.h b/Engine/Plugins/Experimental/NNE/Source/NNERuntimeRDG/Private/Dml/NNEDmlPaddingCommon.h
deleted file mode 100644
