HLSL operator InstanceNormalization performance improvement.

#rnx

diff --git a/Engine/Plugins/Experimental/NNE/Source/NNEHlslShaders/Internal/NNEHlslShadersInstanceNormalizationCS.h b/Engine/Plugins/Experimental/NNE/Source/NNEHlslShaders/Internal/NNEHlslShadersInstanceNormalizationCS.h
--- a/Engine/Plugins/Experimental/NNE/Source/NNEHlslShaders/Internal/NNEHlslShadersInstanceNormalizationCS.h
+++ b/Engine/Plugins/Experimental/NNE/Source/NNEHlslShaders/Internal/NNEHlslShadersInstanceNormalizationCS.h
@@ -13,9 +13,10 @@ namespace UE::NNEHlslShaders::Internal
 	enum class EInstanceNormalizationAlgorithm : uint8
 	{
 		Simple1x265 = 0,
-		SharedMemory8x32,
-		SharedMemory16x16,
-		SharedMemory32x8,
+		SharedMemory256x1,
+		SharedMemory512x1,
+		SharedMemory768x1,
+		SharedMemory1024x1,
 		MAX
 	};
 
diff --git a/Engine/Plugins/Experimental/NNE/Source/NNEHlslShaders/Private/NNEHlslShadersInstanceNormalizationCS.cpp b/Engine/Plugins/Experimental/NNE/Source/NNEHlslShaders/Private/NNEHlslShadersInstanceNormalizationCS.cpp
--- a/Engine/Plugins/Experimental/NNE/Source/NNEHlslShaders/Private/NNEHlslShadersInstanceNormalizationCS.cpp
+++ b/Engine/Plugins/Experimental/NNE/Source/NNEHlslShaders/Private/NNEHlslShadersInstanceNormalizationCS.cpp
@@ -12,10 +12,13 @@ namespace UE::NNEHlslShaders::Internal
 		{
 			switch (Algorithm)
 			{
-			case EInstanceNormalizationAlgorithm::Simple1x265: 			return { 1, 256, 1};
-			case EInstanceNormalizationAlgorithm::SharedMemory8x32:		return { 8,  32, 1};
-			case EInstanceNormalizationAlgorithm::SharedMemory16x16:	return {16,  16, 1};
-			case EInstanceNormalizationAlgorithm::SharedMemory32x8:		return {32,   8, 1};
+			case EInstanceNormalizationAlgorithm::Simple1x265: 			return {    1, 256, 1 };
+			case EInstanceNormalizationAlgorithm::SharedMemory256x1:	return {  256,   1, 1 };
+			case EInstanceNormalizationAlgorithm::SharedMemory512x1:	return {  512,   1, 1 };
+
+			// https://learn.microsoft.com/en-us/windows/win32/direct3dhlsl/sm5-attributes-numthreads
+			case EInstanceNormalizationAlgorithm::SharedMemory768x1:	return {  768,   1, 1 }; // Compute shader version 4.X
+			case EInstanceNormalizationAlgorithm::SharedMemory1024x1:	return { 1024,   1, 1 }; // Compute shader version 5.0
 			}
 
 			check(false);
@@ -43,16 +46,17 @@ namespace UE::NNEHlslShaders::Internal
 
 	EInstanceNormalizationAlgorithm TInstanceNormalizationCS::GetAlgorithm(const TInstanceNormalizationCS::FParameters& Parameters)
 	{
-		return EInstanceNormalizationAlgorithm::SharedMemory16x16;
+		return EInstanceNormalizationAlgorithm::SharedMemory768x1; // note: not all platforms support higher number of threads (yet)
 	}
 
 	void TInstanceNormalizationCS::LexFromString(EInstanceNormalizationAlgorithm& OutValue, const TCHAR* StringVal)
 	{
 		OutValue = EInstanceNormalizationAlgorithm::MAX;
 		if (FCString::Stricmp(StringVal, TEXT("Simple1x265")) == 0) OutValue = EInstanceNormalizationAlgorithm::Simple1x265;
-		else if (FCString::Stricmp(StringVal, TEXT("SharedMemory8x32")) == 0) OutValue = EInstanceNormalizationAlgorithm::SharedMemory8x32;
-		else if (FCString::Stricmp(StringVal, TEXT("SharedMemory16x16")) == 0) OutValue = EInstanceNormalizationAlgorithm::SharedMemory16x16;
-		else if (FCString::Stricmp(StringVal, TEXT("SharedMemory32x8")) == 0) OutValue = EInstanceNormalizationAlgorithm::SharedMemory32x8;
+		else if (FCString::Stricmp(StringVal, TEXT("SharedMemory256x1")) == 0) OutValue = EInstanceNormalizationAlgorithm::SharedMemory256x1;
+		else if (FCString::Stricmp(StringVal, TEXT("SharedMemory512x1")) == 0) OutValue = EInstanceNormalizationAlgorithm::SharedMemory512x1;
+		else if (FCString::Stricmp(StringVal, TEXT("SharedMemory768x1")) == 0) OutValue = EInstanceNormalizationAlgorithm::SharedMemory768x1;
+		else if (FCString::Stricmp(StringVal, TEXT("SharedMemory1024x1")) == 0) OutValue = EInstanceNormalizationAlgorithm::SharedMemory1024x1;
 	}
 
 	IMPLEMENT_GLOBAL_SHADER(TInstanceNormalizationCS, "/NNE/NNEHlslShadersInstanceNormalization.usf", "InstanceNormalization", SF_Compute);
diff --git a/Engine/Plugins/Experimental/NNE/Source/NNEHlslShaders/Shaders/NNEHlslShadersInstanceNormalization.usf b/Engine/Plugins/Experimental/NNE/Source/NNEHlslShaders/Shaders/NNEHlslShadersInstanceNormalization.usf
--- a/Engine/Plugins/Experimental/NNE/Source/NNEHlslShaders/Shaders/NNEHlslShadersInstanceNormalization.usf
+++ b/Engine/Plugins/Experimental/NNE/Source/NNEHlslShaders/Shaders/NNEHlslShadersInstanceNormalization.usf
@@ -20,20 +20,30 @@ RWBuffer<BUFFER_TYPE> Output;	// Output tensor of the same shape as input
 #if ALGORITHM == 0
 #define GROUP_SIZE_X 1
 #define GROUP_SIZE_Y 256
+
 #elif ALGORITHM == 1
-#define GROUP_SIZE_X 8
-#define GROUP_SIZE_Y 32
+#define GROUP_SIZE_X 256
+#define GROUP_SIZE_Y 1
+
 #elif ALGORITHM == 2
-#define GROUP_SIZE_X 16
-#define GROUP_SIZE_Y 16
+#define GROUP_SIZE_X 512
+#define GROUP_SIZE_Y 1
+
 #elif ALGORITHM == 3
-#define GROUP_SIZE_X 32
-#define GROUP_SIZE_Y 8
+#define GROUP_SIZE_X 768
+#define GROUP_SIZE_Y 1
+
+#elif ALGORITHM == 4
+#define GROUP_SIZE_X 1024
+#define GROUP_SIZE_Y 1
+
 #endif
 
 #if GROUP_SIZE_X > 1
 
-groupshared WORK_TYPE SharedMemory[GROUP_SIZE_X * GROUP_SIZE_Y];
+groupshared WORK_TYPE SharedMemoryMean[GROUP_SIZE_X * GROUP_SIZE_Y];
+groupshared WORK_TYPE SharedMemoryMean2[GROUP_SIZE_X * GROUP_SIZE_Y];
+groupshared int SharedMemoryCount[GROUP_SIZE_X * GROUP_SIZE_Y];
 
 [numthreads(GROUP_SIZE_X, GROUP_SIZE_Y, 1)]
 void InstanceNormalization(in const uint3 DispatchThreadID : SV_DispatchThreadID, in const uint3 GroupThreadID : SV_GroupThreadID)
@@ -41,18 +51,29 @@ void InstanceNormalization(in const uint3 DispatchThreadID : SV_DispatchThreadID
 	const int StartIdx = DispatchThreadID.y * W;
 	const int SharedMemoryOffsetY = GroupThreadID.y * GROUP_SIZE_X;
 
-	const WORK_TYPE InvNum = (WORK_TYPE)1 / (WORK_TYPE)W;
-
-	// Calculate mean
+	// Calculate mean and variance using Welford's online algorithm
+	// https://en.wikipedia.org/wiki/Algorithms_for_calculating_variance
 	if (DispatchThreadID.y < NxC)
 	{
-		WORK_TYPE LocalSum = 0;
+		WORK_TYPE Mean = 0;
+		WORK_TYPE Mean2 = 0;
+		int Count = 0;
 		for (int CurrIdx = GroupThreadID.x; CurrIdx < W; CurrIdx += GROUP_SIZE_X)
 		{
-			LocalSum += READ(Input[StartIdx + CurrIdx]);
+			Count += 1;
+
+			WORK_TYPE Value = READ(Input[StartIdx + CurrIdx]);
+
+			WORK_TYPE Delta = Value - Mean;
+			Mean += Delta / (WORK_TYPE)Count;
+
+			WORK_TYPE Delta2 = Value - Mean;
+			Mean2 += Delta * Delta2;
 		}
 
-		SharedMemory[SharedMemoryOffsetY + GroupThreadID.x] = LocalSum;
+		SharedMemoryMean[SharedMemoryOffsetY + GroupThreadID.x] = Mean;
+		SharedMemoryMean2[SharedMemoryOffsetY + GroupThreadID.x] = Mean2;
+		SharedMemoryCount[SharedMemoryOffsetY + GroupThreadID.x] = Count;
 	}
 
 	// Wait until all threads wrote result to shared memory
@@ -61,64 +82,42 @@ void InstanceNormalization(in const uint3 DispatchThreadID : SV_DispatchThreadID
 	// Use one thread to calculate final result
 	if (GroupThreadID.x == 0 && DispatchThreadID.y < NxC)
 	{
-		WORK_TYPE LocalMean = 0;
+		WORK_TYPE Mean = 0;
+		WORK_TYPE Mean2 = 0;
+		int Count = 0;
 
 		UNROLL
 		for (int Idx = 0; Idx < GROUP_SIZE_X; Idx++)
 		{
-			LocalMean += SharedMemory[SharedMemoryOffsetY + Idx];
-		}
+			WORK_TYPE CurrMean = SharedMemoryMean[SharedMemoryOffsetY + Idx];
+			WORK_TYPE CurrMean2 = SharedMemoryMean2[SharedMemoryOffsetY + Idx];
+			int CurrCount = SharedMemoryCount[SharedMemoryOffsetY + Idx];
 
-		SharedMemory[SharedMemoryOffsetY] = LocalMean * InvNum;
-	}
+			int NewCount = Count + CurrCount;
 
-	// Wait until first thread wrote result to shared memory
-	GroupMemoryBarrierWithGroupSync();
+			WORK_TYPE Delta = CurrMean - Mean;
+			Mean = (CurrCount * CurrMean + Count * Mean) / NewCount;
 
-	WORK_TYPE Mean = 0;
-	if (DispatchThreadID.y < NxC)
-	{
-		Mean = SharedMemory[SharedMemoryOffsetY];
-	
-		// Calculate variance
-		WORK_TYPE LocalSum = 0;
-		WORK_TYPE Tmp = 0;
-		for (int CurrIdx = GroupThreadID.x; CurrIdx < W; CurrIdx += GROUP_SIZE_X)
-		{
-			Tmp = READ(Input[StartIdx + CurrIdx]) - Mean;
+			Mean2 += CurrMean2 + Delta * Delta * (WORK_TYPE)Count * (WORK_TYPE)CurrCount / (WORK_TYPE)NewCount;
 
-			LocalSum += Tmp * Tmp;
+			Count = NewCount;
 		}
 
-		SharedMemory[SharedMemoryOffsetY + GroupThreadID.x] = LocalSum;
-	}
-
-	// Wait until all threads wrote result to shared memory
-	GroupMemoryBarrierWithGroupSync();
-
-	// Use one thread to calculate final result
-	if (GroupThreadID.x == 0 && DispatchThreadID.y < NxC)
-	{
-		WORK_TYPE LocalVar = 0;
+		WORK_TYPE Variance = Mean2 / (WORK_TYPE)Count;
 
-		UNROLL
-		for (int Idx = 0; Idx < GROUP_SIZE_X; Idx++)
-		{
-			LocalVar += SharedMemory[SharedMemoryOffsetY + Idx];
-		}
-
-		SharedMemory[SharedMemoryOffsetY] = (WORK_TYPE)1 / sqrt(LocalVar * InvNum + Epsilon);
+		SharedMemoryMean[SharedMemoryOffsetY] = Mean;
+		SharedMemoryMean2[SharedMemoryOffsetY] = (WORK_TYPE)1 / sqrt(Variance + Epsilon);
 	}
 
 	// Wait until first thread wrote result to shared memory
 	GroupMemoryBarrierWithGroupSync();
 
-	WORK_TYPE Variance = 0;
 	if (DispatchThreadID.y < NxC)
 	{
 		const int ChannelIdx = DispatchThreadID.y % C;
 		
-		WORK_TYPE InvStdev = SharedMemory[SharedMemoryOffsetY];
+		WORK_TYPE Mean = SharedMemoryMean[SharedMemoryOffsetY];
+		WORK_TYPE InvStdev = SharedMemoryMean2[SharedMemoryOffsetY];
 
 		// Carry out instance normalization and write to output
 		for (int CurrIdx = GroupThreadID.x; CurrIdx < W; CurrIdx += GROUP_SIZE_X)
@@ -143,26 +142,24 @@ void InstanceNormalization(in const uint3 DispatchThreadID : SV_DispatchThreadID
 
 	const WORK_TYPE InvNum = (WORK_TYPE)1 / (WORK_TYPE)W;
 
-	// Calculate mean
+	// Calculate mean and variance using Welfor's online algorithm
 	WORK_TYPE Mean = 0;
+	WORK_TYPE Mean2 = 0;
+	int Count = 0;
 	for (int Idx = 0; Idx < W; Idx++)
 	{
-		Mean += READ(Input[StartIdx + Idx]);
-	}
+		Count += 1;
 
-	Mean *= InvNum;
+		WORK_TYPE Value = READ(Input[StartIdx + Idx]);
 
-	// Calculate variance
-	WORK_TYPE Variance = 0;
-	WORK_TYPE Tmp = 0;
-	for (int Idx = 0; Idx < W; Idx++)
-	{
-		Tmp = READ(Input[StartIdx + Idx]) - Mean;
+		WORK_TYPE Delta = Value - Mean;
+		Mean += Delta / (WORK_TYPE)Count;
 
-		Variance += Tmp * Tmp;
+		WORK_TYPE Delta2 = Value - Mean;
+		Mean2 += Delta * Delta2;
 	}
 
-	Variance *= InvNum;
+	WORK_TYPE Variance = Mean2 * InvNum;
 
 	WORK_TYPE InvStdev = (WORK_TYPE)1 / sqrt(Variance + Epsilon);
 
diff --git a/Engine/Plugins/Experimental/NNE/Source/NNEQA/Resources/NNEQATestDesc.json b/Engine/Plugins/Experimental/NNE/Source/NNEQA/Resources/NNEQATestDesc.json
--- a/Engine/Plugins/Experimental/NNE/Source/NNEQA/Resources/NNEQATestDesc.json
+++ b/Engine/Plugins/Experimental/NNE/Source/NNEQA/Resources/NNEQATestDesc.json
@@ -1383,9 +1383,10 @@
 			"name": "InstanceNormalization.Algorithms",
 			"attribute_maps": [
 				{ "attributes": [ { "name": "__UE__algorithm", "value": "Simple1x265", "type": "String" } ] },
-				{ "attributes": [ { "name": "__UE__algorithm", "value": "SharedMemory8x32", "type": "String" } ] },
-				{ "attributes": [ { "name": "__UE__algorithm", "value": "SharedMemory16x16", "type": "String" } ] },
-				{ "attributes": [ { "name": "__UE__algorithm", "value": "SharedMemory32x8", "type": "String" } ] }
+				{ "attributes": [ { "name": "__UE__algorithm", "value": "SharedMemory256x1", "type": "String" } ] },
+				{ "attributes": [ { "name": "__UE__algorithm", "value": "SharedMemory512x1", "type": "String" } ] },
+				{ "attributes": [ { "name": "__UE__algorithm", "value": "SharedMemory768x1", "type": "String" } ] },
+				{ "attributes": [ { "name": "__UE__algorithm", "value": "SharedMemory1024x1", "type": "String" } ] }
 			]
 		},
 		{
