NNE RDGDml AvgPool operator + support for opset 15 for Pad

Tests:
* added dedicated tests for averagePool op

#rnx

diff --git a/Engine/Plugins/Experimental/NNE/Source/NNEQA/Resources/NNEQATestDesc.json b/Engine/Plugins/Experimental/NNE/Source/NNEQA/Resources/NNEQATestDesc.json
--- a/Engine/Plugins/Experimental/NNE/Source/NNEQA/Resources/NNEQATestDesc.json
+++ b/Engine/Plugins/Experimental/NNE/Source/NNEQA/Resources/NNEQATestDesc.json
@@ -436,9 +436,10 @@
 			]
 		},
 		{
-			"category": "MaxPool",
+			"category": "PoolOp",
 			"targets": [
-				{ "target": "MaxPool" }
+				{ "target": "MaxPool" },
+				{ "target": "AveragePool", "absolute_tolerance": 1e-6 }
 			],
 			"runtimes": [
 				{ "name": "NNERuntimeRDGHlsl", "skip": true, "comment": "Not supported" }
@@ -1061,13 +1062,13 @@
 			]
 		},
 		{
-			"name": "MaxPool.2D.Default",
+			"name": "PoolOp.2D.Default",
 			"datasets": [
 				{ "inputs": [ { "shape": [ 1, 3, 32, 32 ] } ], "outputs": [ { "shape": [ 1, 3, 31, 31 ] } ] }
 			]
 		},
 		{
-			"name": "MaxPool.2D.Strides2",
+			"name": "PoolOp.2D.Strides2",
 			"datasets": [
 				{ "inputs": [ { "shape": [ 1, 2, 256, 256 ] } ], "outputs": [ { "shape": [ 1, 2, 128, 128 ] } ] }
 			]
@@ -1534,7 +1535,7 @@
 			]
 		},
 		{
-			"name": "MaxPool.2D.Default",
+			"name": "PoolOp.2D.Default",
 			"attribute_maps": [
 				{ "attributes": [ 
 						{ "name": "kernel_shape", "value": [2, 2], "type": "Int32Array" } 
@@ -1543,7 +1544,7 @@
 			]
 		},
 		{
-			"name": "MaxPool.2D.Strides2",
+			"name": "PoolOp.2D.Strides2",
 			"attribute_maps": [
 				{ "attributes": [ 
 						{ "name": "kernel_shape", "value": [2, 2], "type": "Int32Array" },
diff --git a/Engine/Plugins/Experimental/NNE/Source/NNERuntimeRDG/Private/Dml/NNEDmlOperatorMaxPool.cpp b/Engine/Plugins/Experimental/NNE/Source/NNERuntimeRDG/Private/Dml/NNEDmlOperatorMaxPool.cpp
deleted file mode 100644
diff --git a/Engine/Plugins/Experimental/NNE/Source/NNERuntimeRDG/Private/Dml/NNEDmlOperatorPad.cpp b/Engine/Plugins/Experimental/NNE/Source/NNERuntimeRDG/Private/Dml/NNEDmlOperatorPad.cpp
--- a/Engine/Plugins/Experimental/NNE/Source/NNERuntimeRDG/Private/Dml/NNEDmlOperatorPad.cpp
+++ b/Engine/Plugins/Experimental/NNE/Source/NNERuntimeRDG/Private/Dml/NNEDmlOperatorPad.cpp
@@ -41,7 +41,7 @@ public:
 	//
 	virtual bool Initialize(IDMLDevice* Device, TArrayView<const NNECore::Internal::FTensor> InputTensors, TArrayView<const NNECore::Internal::FTensor> OutputTensors, const NNECore::FAttributeMap& Attributes) override
 	{
-		check(InputTensors.Num() == 1);
+		check(InputTensors.Num() >= 1);
 		check(OutputTensors.Num() == 1);
 
 		const NNECore::Internal::FTensor& InputTensor = InputTensors[0];
@@ -54,7 +54,22 @@ public:
 
 		Value = Attributes.GetValueOrDefault<float>(TEXT("value"), 0.0f);
 		Mode = ModeFromString(Attributes.GetValue<FString>(TEXT("mode")));
-		Pads = Attributes.GetValue<TArray<int32>>(TEXT("pads"));
+
+		if (InputTensors.Num() >= 2)
+		{
+			if (!InputTensors[1].HasPreparedData())
+			{
+				UE_LOG(LogNNE, Warning, TEXT("pads is only supported as an attribute or a constant tensor, it is here a variable tensor of name %s."), *InputTensors[1].GetName());
+				return false;
+			}
+			Pads.Append(InputTensors[1].GetPreparedData<int64>());
+		}
+		else
+		{
+			Pads = Attributes.GetValue<TArray<int32>>(TEXT("pads"));
+		}
+
+		
 
 		if (InputTensor.GetShape().Rank() * 2 != Pads.Num())
 		{
diff --git a/Engine/Plugins/Experimental/NNE/Source/NNERuntimeRDG/Private/Dml/NNEDmlOperatorPool.cpp b/Engine/Plugins/Experimental/NNE/Source/NNERuntimeRDG/Private/Dml/NNEDmlOperatorPool.cpp
new file mode 100644
--- /dev/null
+++ b/Engine/Plugins/Experimental/NNE/Source/NNERuntimeRDG/Private/Dml/NNEDmlOperatorPool.cpp
@@ -0,0 +1,193 @@
+// Copyright Epic Games, Inc. All Rights Reserved.
+
+#ifdef NNE_USE_DIRECTML
+
+#include "NNEDmlOperator.h"
+#include "NNEDmlOperatorUtils.h"
+
+#include "Algo/AnyOf.h"
+
+namespace UE::NNERuntimeRDG::Private::Dml
+{
+
+/**
+	* MaxPool
+	*/
+template <typename TDmlPoolOpDesc, DML_OPERATOR_TYPE DmlOpType>
+class FOperatorDmlPool : public FOperatorDml
+{
+	using FIntArray = TArray<int32>;
+
+public:
+
+	static FOperatorDml* Create()
+	{
+		return new FOperatorDmlPool<TDmlPoolOpDesc, DmlOpType>();
+	}
+
+	//
+	//
+	//
+	virtual bool Initialize(IDMLDevice* Device, TArrayView<const NNECore::Internal::FTensor> InputTensors, TArrayView<const NNECore::Internal::FTensor> OutputTensors, const NNECore::FAttributeMap& Attributes) override
+	{
+		//TODO: int64 attributes
+		check(InputTensors.Num() == 1);
+		TConstArrayView<uint32> InputShape = InputTensors[0].GetShape().GetData();
+		check(InputShape.Num() >= 2);
+		check(OutputTensors.Num() == 1 || OutputTensors.Num() == 2);
+
+		const uint32 NumSpatialDimensions = InputShape.Num() - NonspatialDimensionCount;
+		if (NumSpatialDimensions != 2 && NumSpatialDimensions != 3)
+		{
+			UE_LOG(LogNNE, Error, TEXT("Number of spatial dimensions must be in range [2, 3] for DML inference."));
+			return false;
+		}
+
+		DmlUtil::FSmallUIntArray StartPadding, EndPadding, KernelShape, Strides, Dilations;
+
+		Strides.Init(1u, NumSpatialDimensions);
+		if (!DmlUtil::GetArrayAttributeNoOverflow(Attributes.GetAttributeValue(TEXT("strides")), Strides, TConstArrayView<uint32>(Strides)))
+		{
+			UE_LOG(LogNNE, Error, TEXT("Strides attribute cast led to overflow"));
+			return false;
+		}
+		check(Strides.Num() == NumSpatialDimensions);
+
+		Dilations.Init(1u, NumSpatialDimensions);
+		if (!DmlUtil::GetArrayAttributeNoOverflow(Attributes.GetAttributeValue(TEXT("dilations")), Dilations, TConstArrayView<uint32>(Dilations)))
+		{
+			UE_LOG(LogNNE, Error, TEXT("Dilations attribute cast led to overflow"));
+			return false;
+		}
+		check(Dilations.Num() == NumSpatialDimensions);
+
+		if (!DmlUtil::GetArrayAttributeNoOverflow(Attributes.GetAttributeValue(TEXT("kernel_shape")), KernelShape))
+		{
+			UE_LOG(LogNNE, Error, TEXT("kernel_shape attribute cast led to overflow"));
+			return false;
+		}
+		if (KernelShape.IsEmpty())
+		{
+			UE_LOG(LogNNE, Error, TEXT("kernel_shape attribute is required for pooling operators"));
+			return false;
+		}
+		check(KernelShape.Num() == NumSpatialDimensions);
+
+		ComputeStartEndPaddings(
+			InputShape,
+			Attributes,
+			StartPadding,
+			EndPadding,
+			KernelPadding(InputShape, KernelShape, Dilations, Strides)
+		);
+
+		for (int Idx = 0; Idx < OutputTensors.Num(); ++Idx)
+		{
+			check(OutputTensors[Idx].GetShape().Rank() == InputShape.Num());
+			check(OutputTensors[Idx].GetShape().GetData()[0] == InputShape[0]);
+			check(OutputTensors[Idx].GetShape().GetData()[1] == InputShape[1]);
+
+			for (uint32 Dim = 0; Dim < NumSpatialDimensions; ++Dim) {
+				uint32 PaddedSize = InputShape[Dim + NonspatialDimensionCount] + StartPadding[Dim] + EndPadding[Dim];
+				check(OutputTensors[Idx].GetShape().GetData()[Dim + NonspatialDimensionCount] == (PaddedSize - KernelShape[Dim]) / Strides[Dim] + 1);
+			}
+		}
+
+		//storage_order is not supported
+		int32 StorageOrder =
+			Attributes.GetValueOrDefault<int32>(
+				TEXT("storage_order"),
+				0);
+		if(StorageOrder != 0)
+		{
+			UE_LOG(LogNNE, Error, TEXT("storage_order != 0 is not supported for DML inference"));
+			return false;
+		}
+
+
+		DmlUtil::FTensorDesc DmlInputTensorDesc;
+		if (!DmlInputTensorDesc.InitFromTensor(InputTensors[0], 4))
+		{
+			UE_LOG(LogNNE, Error, TEXT("Failed to initialize pooling operator's input tensor for DML inference"));
+			return false;
+		}
+
+		DmlUtil::FTensorDesc DmlOutputTensorDesc;
+		if (!DmlOutputTensorDesc.InitFromTensor(OutputTensors[0], 4))
+		{
+			UE_LOG(LogNNE, Error, TEXT("Failed to initialize pooling operator's output tensor for DML inference"));
+			return false;
+		}
+
+		auto FillPoolingDesc = [&](auto& PoolingDesc)
+		{
+			PoolingDesc.InputTensor = &DmlInputTensorDesc.Desc;
+			PoolingDesc.OutputTensor = &DmlOutputTensorDesc.Desc;
+			PoolingDesc.DimensionCount = NumSpatialDimensions;
+			PoolingDesc.WindowSize = KernelShape.GetData();
+			PoolingDesc.Strides = Strides.GetData();
+			PoolingDesc.StartPadding = StartPadding.GetData();
+			PoolingDesc.EndPadding = EndPadding.GetData();
+		};
+
+		const bool bHasDilations =
+			Algo::AnyOf(
+				Dilations,
+				[](auto Dim) {return Dim != 1u; }
+		);
+		const bool bHasOutputIndices = OutputTensors.Num() > 1;
+
+		TDmlPoolOpDesc DmlPoolOpDesc = {};
+		FillPoolingDesc(DmlPoolOpDesc);
+
+		if (DmlOpType == DML_OPERATOR_MAX_POOLING && (bHasOutputIndices || bHasDilations))
+		{
+			DML_MAX_POOLING2_OPERATOR_DESC DmlMaxPoolOpDesc = {};
+
+			if (bHasOutputIndices)
+			{
+				DmlUtil::FTensorDesc DmlIndicesTensorDesc;
+				if (!DmlIndicesTensorDesc.InitFromTensor(OutputTensors[1], 4))
+				{
+					UE_LOG(LogNNE, Warning, TEXT("Failed to initialize MaxPool's indices output tensor for DML inference"));
+					return false;
+				}
+				DmlIndicesTensorDesc.BuffDesc.DataType = DML_TENSOR_DATA_TYPE::DML_TENSOR_DATA_TYPE_UINT64;
+				DmlMaxPoolOpDesc.OutputIndicesTensor = &DmlIndicesTensorDesc.Desc;
+			}
+
+			DmlMaxPoolOpDesc.Dilations = Dilations.GetData();
+			FillPoolingDesc(DmlMaxPoolOpDesc);
+			return CreateOperator(Device, DML_OPERATOR_DESC{ DML_OPERATOR_MAX_POOLING2, &DmlMaxPoolOpDesc });
+		}
+
+		if constexpr (DmlOpType == DML_OPERATOR_AVERAGE_POOLING)
+		{
+			int32 IncludePadding = Attributes.GetValueOrDefault(TEXT("count_include_pad"), 0);
+			DmlPoolOpDesc.IncludePadding = static_cast<BOOL>(IncludePadding);
+		}
+			
+		return CreateOperator(Device, DML_OPERATOR_DESC{ DmlOpType, &DmlPoolOpDesc });
+	}
+};
+
+#define NNE_DML_REGISTER_POOLING_OP(OpName, DmlPrefix) \
+struct FDmlOperator##OpName##Registrator \
+{ \
+	FDmlOperator##OpName##Registrator() \
+	{ \
+		FOperatorRegistryDml::Get()->OpAdd(TEXT(#OpName), FOperatorDmlPool<DML_##DmlPrefix##_POOLING_OPERATOR_DESC, DML_OPERATOR_##DmlPrefix##_POOLING>::Create); \
+	} \
+}; \
+\
+static FDmlOperator##OpName##Registrator RegisterDmlOperator##OpName;
+
+// Register pooling operator on Module startup
+NNE_DML_REGISTER_POOLING_OP(MaxPool, MAX)
+NNE_DML_REGISTER_POOLING_OP(AveragePool, AVERAGE)
+
+#undef NNE_DML_REGISTER_POOLING_OP
+
+} // namespace UE::NNERuntimeRDG::Private::Dml
+
+#endif // NNE_USE_DIRECTML
