NNE - Refactor RDG Dml element-wise and activations

#rnx

diff --git a/Engine/Plugins/Experimental/NNE/Source/NNERuntimeRDG/Private/Dml/NNEDmlOperator.h b/Engine/Plugins/Experimental/NNE/Source/NNERuntimeRDG/Private/Dml/NNEDmlOperator.h
--- a/Engine/Plugins/Experimental/NNE/Source/NNERuntimeRDG/Private/Dml/NNEDmlOperator.h
+++ b/Engine/Plugins/Experimental/NNE/Source/NNERuntimeRDG/Private/Dml/NNEDmlOperator.h
@@ -7,7 +7,6 @@
 #include "NNEDmlCommon.h"
 #include "NNECoreTypes.h"
 #include "NNECoreTensor.h"
-#include "NNECoreOperator.h"
 #include "NNECoreAttributeMap.h"
 #include "NNERuntimeRDG.h"
 
diff --git a/Engine/Plugins/Experimental/NNE/Source/NNERuntimeRDG/Private/Dml/NNEDmlOperatorActivations.cpp b/Engine/Plugins/Experimental/NNE/Source/NNERuntimeRDG/Private/Dml/NNEDmlOperatorActivations.cpp
new file mode 100644
--- /dev/null
+++ b/Engine/Plugins/Experimental/NNE/Source/NNERuntimeRDG/Private/Dml/NNEDmlOperatorActivations.cpp
@@ -0,0 +1,249 @@
+// Copyright Epic Games, Inc. All Rights Reserved.
+#ifdef NNE_USE_DIRECTML
+#include "NNEDmlOperator.h"
+
+namespace UE::NNERuntimeRDG::Private::Dml
+{
+
+/**
+ * Activation unary ML operator implementation
+ */
+template
+<
+	typename DmlActivationOpDescType,
+	DML_OPERATOR_TYPE DmlActivationOpType
+>
+class FOperatorDmlActivationUnary : public FOperatorDml
+{
+public:
+
+	static FOperatorDml* Create()
+	{
+		return new FOperatorDmlActivationUnary();
+	}
+
+	virtual ~FOperatorDmlActivationUnary() = default;
+
+private:
+
+	FOperatorDmlActivationUnary() : Alpha(0.0f), Beta(0.0f), Gamma(0.0f) {}
+	float Alpha;
+	float Beta;
+	float Gamma;
+
+public:
+
+	//
+	//
+	//
+	virtual bool Initialize(IDMLDevice* Device, TArrayView<const NNECore::Internal::FTensor> InputTensors, TArrayView<const NNECore::Internal::FTensor> OutputTensors, const NNECore::FAttributeMap& Attributes) override
+	{
+		const NNECore::Internal::FTensor& InputTensorDesc = InputTensors[0];
+		const NNECore::Internal::FTensor& OutputTensorDesc = OutputTensors[0];
+
+		Alpha = Attributes.GetValueOrDefault(TEXT("alpha"), Alpha);
+		Beta = Attributes.GetValueOrDefault(TEXT("beta"), Beta);
+		Gamma = Attributes.GetValueOrDefault(TEXT("gamma"), Gamma);
+
+
+		// Initialize tensor descriptor (it's same for both input and output)
+		DmlUtil::FTensorDesc	DmlTensorDesc{};
+
+		if (!InitDmlTensorDesc(DmlTensorDesc, InputTensorDesc))
+		{
+			UE_LOG(LogNNE, Warning, TEXT("Failed to initialize tensor(s) for DML inference"));
+			return false;
+		}
+
+		DmlActivationOpDescType	DmlActivationOpDesc{};
+
+		InitDmlOpDesc(DmlActivationOpDesc, DmlTensorDesc);
+
+		DML_OPERATOR_DESC DmlOpDesc{};
+
+		DmlOpDesc.Type = DmlActivationOpType;
+		DmlOpDesc.Desc = &DmlActivationOpDesc;
+
+		return CreateOperator(Device, DmlOpDesc);
+	}
+
+private:
+
+	template<typename OpDesc>
+	void InitDmlOpDesc(OpDesc& Desc, DmlUtil::FTensorDesc& TensorDesc)
+	{
+		Desc.InputTensor = &TensorDesc.Desc;
+		Desc.OutputTensor = &TensorDesc.Desc;
+	}
+
+	void InitDmlOpDesc(DML_ACTIVATION_SOFTPLUS_OPERATOR_DESC& Desc, DmlUtil::FTensorDesc& TensorDesc)
+	{
+		Desc.InputTensor = &TensorDesc.Desc;
+		Desc.OutputTensor = &TensorDesc.Desc;
+		Desc.Steepness = 1.0f;
+	}
+
+	void InitDmlOpDesc(DML_ACTIVATION_SCALED_ELU_OPERATOR_DESC& Desc, DmlUtil::FTensorDesc& TensorDesc)
+	{
+		Desc.InputTensor = &TensorDesc.Desc;
+		Desc.OutputTensor = &TensorDesc.Desc;
+		Desc.Alpha = Alpha;
+		Desc.Gamma = Gamma;
+	}
+
+	void InitDmlOpDesc(DML_ACTIVATION_ELU_OPERATOR_DESC& Desc, DmlUtil::FTensorDesc& TensorDesc)
+	{
+		Desc.InputTensor = &TensorDesc.Desc;
+		Desc.OutputTensor = &TensorDesc.Desc;
+		Desc.Alpha = Alpha;
+	}
+
+	void InitDmlOpDesc(DML_ACTIVATION_HARD_SIGMOID_OPERATOR_DESC& Desc, DmlUtil::FTensorDesc& TensorDesc)
+	{
+		Desc.InputTensor = &TensorDesc.Desc;
+		Desc.OutputTensor = &TensorDesc.Desc;
+		Desc.Alpha = Alpha;
+		Desc.Beta = Beta;
+	}
+
+	void InitDmlOpDesc(DML_ACTIVATION_LEAKY_RELU_OPERATOR_DESC& Desc, DmlUtil::FTensorDesc& TensorDesc)
+	{
+		Desc.InputTensor = &TensorDesc.Desc;
+		Desc.OutputTensor = &TensorDesc.Desc;
+		Desc.Alpha = Alpha;
+	}
+};
+
+/**
+ * Activation binary ML operator implementation
+ */
+template
+<
+	typename DmlActivationOpDescType,
+	DML_OPERATOR_TYPE DmlActivationOpType
+>
+class FOperatorDmlActivationBinary : public FOperatorDml
+{
+
+public:
+
+	static FOperatorDml* Create()
+	{
+		return new FOperatorDmlActivationBinary();
+	}
+
+private:
+
+	FOperatorDmlActivationBinary() = default;
+
+public:
+
+	//
+	//
+	//
+	virtual bool Initialize(IDMLDevice* Device, TArrayView<const NNECore::Internal::FTensor> InputTensors, TArrayView<const NNECore::Internal::FTensor> OutputTensors, const NNECore::FAttributeMap& Attributes) override
+	{
+		const NNECore::Internal::FTensor& InputATensorDesc = InputTensors[0];
+		const NNECore::Internal::FTensor& InputBTensorDesc = InputTensors[1];
+		const NNECore::Internal::FTensor& OutputTensorDesc = OutputTensors[0];
+
+		// Initialize tensor descriptors
+		DmlUtil::FTensorDesc	DmlInputATensorDesc{};
+		DmlUtil::FTensorDesc	DmlInputBTensorDesc{};
+		DmlUtil::FTensorDesc	DmlOutputTensorDesc{};
+
+		if (!InitDmlTensorDesc(DmlInputATensorDesc, InputATensorDesc, OutputTensorDesc))
+		{
+			UE_LOG(LogNNE, Warning, TEXT("Failed to initialize tensor(s) for DML inference"));
+			return false;
+		}
+
+		if (!InitDmlTensorDesc(DmlInputBTensorDesc, InputBTensorDesc, OutputTensorDesc))
+		{
+			UE_LOG(LogNNE, Warning, TEXT("Failed to initialize tensor(s) for DML inference"));
+			return false;
+		}
+
+		if (!InitDmlTensorDesc(DmlOutputTensorDesc, OutputTensorDesc))
+		{
+			UE_LOG(LogNNE, Warning, TEXT("Failed to initialize tensor(s) for DML inference"));
+			return false;
+		}
+
+		DmlActivationOpDescType	DmlActivationOpDesc{};
+
+		InitDmlOpDesc(DmlActivationOpDesc, DmlInputATensorDesc, DmlInputBTensorDesc, DmlOutputTensorDesc);
+
+		DML_OPERATOR_DESC DmlOpDesc{};
+
+		DmlOpDesc.Type = DmlActivationOpType;
+		DmlOpDesc.Desc = &DmlActivationOpDesc;
+
+		return CreateOperator(Device, DmlOpDesc);
+	}
+
+private:
+
+	template<typename OpDesc>
+	void InitDmlOpDesc(OpDesc& Desc, DmlUtil::FTensorDesc& LHSTensor, DmlUtil::FTensorDesc& RHSTensor, DmlUtil::FTensorDesc& OutputTensor)
+	{
+		Desc.ATensor = &LHSTensor.Desc;
+		Desc.BTensor = &RHSTensor.Desc;
+		Desc.OutputTensor = &OutputTensor.Desc;
+	}
+
+	void InitDmlOpDesc(DML_ACTIVATION_PARAMETERIZED_RELU_OPERATOR_DESC& Desc, DmlUtil::FTensorDesc& LHSTensor, DmlUtil::FTensorDesc& RHSTensor, DmlUtil::FTensorDesc& OutputTensor)
+	{
+		Desc.InputTensor = &LHSTensor.Desc;
+		Desc.SlopeTensor = &RHSTensor.Desc;
+		Desc.OutputTensor = &OutputTensor.Desc;
+	}
+};
+
+#define REGISTER_OP_ACTIVATION(OpName, DmlOpName, OpClass) \
+struct FDmlOperator##OpName##Registrator \
+{ \
+	FDmlOperator##OpName##Registrator() \
+	{ \
+		FOperatorRegistryDml::Get()->OpAdd(TEXT(#OpName), OpClass<DML_ACTIVATION_##DmlOpName##_OPERATOR_DESC, DML_OPERATOR_ACTIVATION_##DmlOpName>::Create); \
+	} \
+}; \
+\
+static FDmlOperator##OpName##Registrator RegisterDmlOperator##OpName;
+
+
+#define REGISTER_OP_ACTIVATION_UNARY(OpName, DmlOpName) REGISTER_OP_ACTIVATION(OpName, DmlOpName, FOperatorDmlActivationUnary)
+
+#define REGISTER_OP_ACTIVATION_UNARY_PARAMS(OpName, DmlOpName, InAlpha, InBeta, InGamma) \
+template<> FOperatorDmlActivationUnary<DML_ACTIVATION_##DmlOpName##_OPERATOR_DESC, DML_OPERATOR_ACTIVATION_##DmlOpName>::FOperatorDmlActivationUnary() \
+	: Alpha(InAlpha), Beta(InBeta), Gamma(InGamma) \
+{ \
+} \
+REGISTER_OP_ACTIVATION_UNARY(OpName, DmlOpName)
+
+
+#define REGISTER_OP_ACTIVATION_BINARY(OpName, DmlOpName) REGISTER_OP_ACTIVATION(OpName, DmlOpName, FOperatorDmlActivationBinary)
+
+
+
+// 	REGISTER_OP_ACTIVATION_UNARY(			OpName, 			DmlOpName )
+// 	REGISTER_OP_ACTIVATION_UNARY_PARAMS(	OpName, 			DmlOpName, 				InAlpha, 							InBeta, 	InGamma )
+
+	REGISTER_OP_ACTIVATION_UNARY_PARAMS(	Elu, 				ELU, 					1.0f, 								0.0f, 		1.05070102214813232421875f )
+	REGISTER_OP_ACTIVATION_UNARY_PARAMS(	HardSigmoid, 		HARD_SIGMOID, 			0.2f, 								0.5f, 		0.0f )
+	REGISTER_OP_ACTIVATION_UNARY_PARAMS(	LeakyRelu, 			LEAKY_RELU, 			0.01f, 								0.0f, 		0.0f )
+	REGISTER_OP_ACTIVATION_UNARY(			Relu, 				RELU )
+	REGISTER_OP_ACTIVATION_UNARY_PARAMS(	Selu, 				SCALED_ELU, 			1.67326319217681884765625f, 		0.0f, 		1.05070102214813232421875f )
+	REGISTER_OP_ACTIVATION_UNARY(			Sigmoid, 			SIGMOID	)
+	REGISTER_OP_ACTIVATION_UNARY(			Softplus, 			SOFTPLUS )
+	REGISTER_OP_ACTIVATION_UNARY(			Softsign, 			SOFTSIGN )
+
+// 	REGISTER_OP_ACTIVATION_BINARY(	OpName, 			DmlOpName )
+
+	REGISTER_OP_ACTIVATION_BINARY(	Prelu, 				PARAMETERIZED_RELU )
+
+
+
+} // namespace UE::NNERuntimeRDG::Private::Dml
+
+#endif // NNE_USE_DIRECTML
\ No newline at end of file
diff --git a/Engine/Plugins/Experimental/NNE/Source/NNERuntimeRDG/Private/Dml/NNEDmlOperatorElementwise.cpp b/Engine/Plugins/Experimental/NNE/Source/NNERuntimeRDG/Private/Dml/NNEDmlOperatorElementwise.cpp
--- a/Engine/Plugins/Experimental/NNE/Source/NNERuntimeRDG/Private/Dml/NNEDmlOperatorElementwise.cpp
+++ b/Engine/Plugins/Experimental/NNE/Source/NNERuntimeRDG/Private/Dml/NNEDmlOperatorElementwise.cpp
@@ -5,125 +5,6 @@
 namespace UE::NNERuntimeRDG::Private::Dml
 {
 
-/**
- * Utility function to get operator type
- */
-template<typename TElementWiseOpDesc>
-DML_OPERATOR_TYPE GetDmlElementWiseUnaryOpType();
-
-#define OP_EW(OpName) \
-template<> \
-DML_OPERATOR_TYPE GetDmlElementWiseUnaryOpType<DML_ELEMENT_WISE_##OpName##_OPERATOR_DESC>() { return DML_OPERATOR_ELEMENT_WISE_##OpName; }
-
-#define OP_AN(OpName) \
-template<> \
-DML_OPERATOR_TYPE GetDmlElementWiseUnaryOpType<DML_ACTIVATION_##OpName##_OPERATOR_DESC>() { return DML_OPERATOR_ACTIVATION_##OpName; }
-
-
-OP_EW(IDENTITY)
-OP_EW(ABS)
-OP_EW(ACOS)
-OP_EW(ACOSH)
-OP_EW(ASIN)
-OP_EW(ASINH)
-OP_EW(ATAN)
-OP_EW(ATANH)
-// BitShift
-// Cast
-OP_EW(CEIL)
-OP_EW(CLIP)
-OP_EW(COS)
-OP_EW(COSH)
-OP_AN(ELU)
-OP_EW(ERF)
-OP_EW(EXP)
-OP_EW(FLOOR)
-OP_EW(IS_INFINITY)
-OP_EW(IS_NAN)
-OP_AN(HARDMAX)
-OP_AN(HARD_SIGMOID)
-OP_AN(LEAKY_RELU)
-OP_EW(LOG)
-OP_EW(NEGATE)
-// Not
-OP_EW(RECIP)
-OP_AN(RELU)
-OP_EW(ROUND)
-OP_AN(SCALED_ELU)
-OP_AN(SIGMOID)
-OP_EW(SIGN)
-OP_EW(SIN)
-OP_EW(SINH)
-OP_AN(SOFTPLUS)
-OP_AN(SOFTSIGN)
-OP_EW(SQRT)
-OP_EW(TAN)
-OP_EW(TANH)
-
-#undef OP_EW
-#undef OP_AN
-
-
-/**
- * Utility function to get operator type
- */
-//template<typename TActivationOpDesc>
-//DML_OPERATOR_TYPE GetDmlActivationOpType();
-//
-//#define OP(OpName) \
-//template<> \
-//DML_OPERATOR_TYPE GetDmlActivationOpType<DML_ACTIVATION_##OpName##_OPERATOR_DESC>() { return DML_OPERATOR_ACTIVATION_##OpName; }
-//
-//OP(ELU)
-//OP(ERF)
-//OP(HARDMAX)
-//OP(HARD_SIGMOID)
-//OP(LEAKY_RELU)
-//OP(LINEAR)
-//OP(LOG_SOFTMAX)
-//OP(PARAMETERIZED_RELU)
-//OP(PARAMETRIC_SOFTPLUS)
-//OP(RELU)
-//OP(SCALED_ELU)
-//OP(SCALED_TANH)
-//OP(SOFTMAX)
-//OP(SOFTPLUS)
-//OP(SOFTSIGN)
-//OP(TANH)
-//OP(THRESHOLDED_RELU)
-//
-
-/**
- * Utility function to get operator type
- */
-template<typename TElementWiseOpDesc>
-DML_OPERATOR_TYPE GetDmlElementWiseBinaryOpType();
-
-#define OP_EW(OpName) \
-template<> \
-DML_OPERATOR_TYPE GetDmlElementWiseBinaryOpType<DML_ELEMENT_WISE_##OpName##_OPERATOR_DESC>() { return DML_OPERATOR_ELEMENT_WISE_##OpName; }
-
-#define OP_AN(OpName) \
-template<> \
-DML_OPERATOR_TYPE GetDmlElementWiseBinaryOpType<DML_ACTIVATION_##OpName##_OPERATOR_DESC>() { return DML_OPERATOR_ACTIVATION_##OpName; }
-
-OP_EW(ADD)
-//OP_EW(LOGICAL_AND)
-OP_EW(DIVIDE)
-//OP_EW(LOGICAL_EQUALS)
-//OP_EW(LOGICAL_GREATER_THAN)
-//OP_EW(LOGICAL_LESS_THAN)
-//OP_EW(MOD)
-OP_EW(MULTIPLY)
-//OP_EW(LOGICAL_OR)
-OP_AN(PARAMETERIZED_RELU)
-OP_EW(POW)
-OP_EW(SUBTRACT)
-//OP_EW(LOGICAL_XOR)
-
-#undef OP_EW
-#undef OP_AN
-
 
 /**
  * Element-wise unary ML operator implementation
@@ -131,7 +12,7 @@ OP_EW(SUBTRACT)
 template
 <
 	typename DmlElementWiseOpDescType, 
-	NNECore::Internal::EElementWiseUnaryOperatorType OpType
+	DML_OPERATOR_TYPE DmlElementWiseOpType
 >
 class FOperatorDmlElementWiseUnary : public FOperatorDml
 {
@@ -146,13 +27,9 @@ public:
 
 private:
 
-	FOperatorDmlElementWiseUnary() : Alpha(0.0f), Beta(0.0f), Gamma(0.0f), Min(TNumericLimits<float>::Min()), Max(TNumericLimits<float>::Max()), Num(1) {}
-	float Alpha;
-	float Beta;
-	float Gamma;
+	FOperatorDmlElementWiseUnary() : Min(TNumericLimits<float>::Min()), Max(TNumericLimits<float>::Max()) {}
 	float Min;
 	float Max;
-	uint32 Num;
 
 public:
 
@@ -161,14 +38,10 @@ public:
 	//
 	virtual bool Initialize(IDMLDevice* Device, TArrayView<const NNECore::Internal::FTensor> InputTensors, TArrayView<const NNECore::Internal::FTensor> OutputTensors, const NNECore::FAttributeMap& Attributes) override
 	{
-		Num = InputTensors[0].GetVolume();
 
 		const NNECore::Internal::FTensor& InputTensorDesc = InputTensors[0];
 		const NNECore::Internal::FTensor& OutputTensorDesc = OutputTensors[0];
 
-		Alpha = Attributes.GetValueOrDefault(TEXT("alpha"), Alpha);
-		Beta = Attributes.GetValueOrDefault(TEXT("beta"), Beta);
-		Gamma = Attributes.GetValueOrDefault(TEXT("gamma"), Gamma);
 
 		if constexpr (std::is_same_v<DmlElementWiseOpDescType, DML_ELEMENT_WISE_CLIP_OPERATOR_DESC>)
 		{
@@ -210,7 +83,7 @@ public:
 
 		DML_OPERATOR_DESC DmlOpDesc{};
 
-		DmlOpDesc.Type = GetDmlElementWiseUnaryOpType<DmlElementWiseOpDescType>();
+		DmlOpDesc.Type = DmlElementWiseOpType;
 		DmlOpDesc.Desc = &DmlElemWiseOpDesc;
 
 		return CreateOperator(Device, DmlOpDesc);
@@ -232,76 +105,16 @@ private:
 		Desc.Min = Min;
 		Desc.Max = Max;
 	}
-
-	void InitDmlOpDesc(DML_ACTIVATION_SOFTPLUS_OPERATOR_DESC& Desc, DmlUtil::FTensorDesc& TensorDesc)
-	{
-		Desc.InputTensor = &TensorDesc.Desc;
-		Desc.OutputTensor = &TensorDesc.Desc;
-		Desc.Steepness = 1.0f;
-	}
-
-	void InitDmlOpDesc(DML_ACTIVATION_SCALED_ELU_OPERATOR_DESC& Desc, DmlUtil::FTensorDesc& TensorDesc)
-	{
-		Desc.InputTensor = &TensorDesc.Desc;
-		Desc.OutputTensor = &TensorDesc.Desc;
-		Desc.Alpha = Alpha;
-		Desc.Gamma = Gamma;
-	}
-
-	void InitDmlOpDesc(DML_ACTIVATION_ELU_OPERATOR_DESC& Desc, DmlUtil::FTensorDesc& TensorDesc)
-	{
-		Desc.InputTensor = &TensorDesc.Desc;
-		Desc.OutputTensor = &TensorDesc.Desc;
-		Desc.Alpha = Alpha;
-	}
-
-	void InitDmlOpDesc(DML_ACTIVATION_HARD_SIGMOID_OPERATOR_DESC& Desc, DmlUtil::FTensorDesc& TensorDesc)
-	{
-		Desc.InputTensor = &TensorDesc.Desc;
-		Desc.OutputTensor = &TensorDesc.Desc;
-		Desc.Alpha = Alpha;
-		Desc.Beta = Beta;
-	}
-
-	void InitDmlOpDesc(DML_ACTIVATION_LEAKY_RELU_OPERATOR_DESC& Desc, DmlUtil::FTensorDesc& TensorDesc)
-	{
-		Desc.InputTensor = &TensorDesc.Desc;
-		Desc.OutputTensor = &TensorDesc.Desc;
-		Desc.Alpha = Alpha;
-	}
 };
 
-template<> FOperatorDmlElementWiseUnary<DML_ELEMENT_WISE_CLIP_OPERATOR_DESC, NNECore::Internal::EElementWiseUnaryOperatorType::Clip>::FOperatorDmlElementWiseUnary()
-{
-}
-
-template<> FOperatorDmlElementWiseUnary<DML_ACTIVATION_SCALED_ELU_OPERATOR_DESC, NNECore::Internal::EElementWiseUnaryOperatorType::Selu>::FOperatorDmlElementWiseUnary()
-	: Alpha(1.67326319217681884765625f), Beta(0.0f), Gamma(1.05070102214813232421875f), Num(1)
-{
-}
-
-template<> FOperatorDmlElementWiseUnary<DML_ACTIVATION_ELU_OPERATOR_DESC, NNECore::Internal::EElementWiseUnaryOperatorType::Elu>::FOperatorDmlElementWiseUnary()
-	: Alpha(1.0f), Beta(0.0f), Gamma(1.05070102214813232421875f), Num(1)
-{
-}
-
-template<> FOperatorDmlElementWiseUnary<DML_ACTIVATION_HARD_SIGMOID_OPERATOR_DESC, NNECore::Internal::EElementWiseUnaryOperatorType::HardSigmoid>::FOperatorDmlElementWiseUnary()
-	: Alpha(0.2f), Beta(0.5f), Gamma(0.0f), Num(1)
-{
-}
-
-template<> FOperatorDmlElementWiseUnary<DML_ACTIVATION_LEAKY_RELU_OPERATOR_DESC, NNECore::Internal::EElementWiseUnaryOperatorType::LeakyRelu>::FOperatorDmlElementWiseUnary()
-	: Alpha(0.01f), Beta(0.0f), Gamma(0.0f), Num(1)
-{
-}
 
 /**
  * Element-wise binary ML operator implementation
  */
 template
 <
-	typename TDmlElementWiseOpDescType,
-	NNECore::Internal::EElementWiseBinaryOperatorType OpType
+	typename DmlElementWiseOpDescType, 
+	DML_OPERATOR_TYPE DmlElementWiseOpType
 >
 class FOperatorDmlElementWiseBinary : public FOperatorDml
 {
@@ -317,8 +130,6 @@ private:
 
 	FOperatorDmlElementWiseBinary() = default;
 
-	uint32 Num { 1 };
-
 public:
 
 	//
@@ -326,8 +137,6 @@ public:
 	//
 	virtual bool Initialize(IDMLDevice* Device, TArrayView<const NNECore::Internal::FTensor> InputTensors, TArrayView<const NNECore::Internal::FTensor> OutputTensors, const NNECore::FAttributeMap& Attributes) override
 	{
-		Num = OutputTensors[0].GetVolume();
-
 		const NNECore::Internal::FTensor& InputATensorDesc = InputTensors[0];
 		const NNECore::Internal::FTensor& InputBTensorDesc = InputTensors[1];
 		const NNECore::Internal::FTensor& OutputTensorDesc = OutputTensors[0];
@@ -355,13 +164,13 @@ public:
 			return false;
 		}
 
-		TDmlElementWiseOpDescType	DmlElemWiseOpDesc{};
+		DmlElementWiseOpDescType	DmlElemWiseOpDesc{};
 
 		InitDmlOpDesc(DmlElemWiseOpDesc, DmlInputATensorDesc, DmlInputBTensorDesc, DmlOutputTensorDesc);
 
 		DML_OPERATOR_DESC DmlOpDesc{};
 
-		DmlOpDesc.Type = GetDmlElementWiseBinaryOpType<TDmlElementWiseOpDescType>();
+		DmlOpDesc.Type = DmlElementWiseOpType;
 		DmlOpDesc.Desc = &DmlElemWiseOpDesc;
 
 		return CreateOperator(Device, DmlOpDesc);
@@ -383,92 +192,60 @@ private:
 		Desc.ExponentTensor = &RHSTensor.Desc;
 		Desc.OutputTensor = &OutputTensor.Desc;
 	}
-
-	void InitDmlOpDesc(DML_ACTIVATION_PARAMETERIZED_RELU_OPERATOR_DESC& Desc, DmlUtil::FTensorDesc& LHSTensor, DmlUtil::FTensorDesc& RHSTensor, DmlUtil::FTensorDesc& OutputTensor)
-	{
-		Desc.InputTensor = &LHSTensor.Desc;
-		Desc.SlopeTensor = &RHSTensor.Desc;
-		Desc.OutputTensor = &OutputTensor.Desc;
-	}
 };
 
-void RegisterElementWiseUnaryOperators()
-{
-#define OP(DmlOpDesc, OpName) FOperatorRegistryDml::Get()->OpAdd(TEXT(#OpName), FOperatorDmlElementWiseUnary<DmlOpDesc, NNECore::Internal::EElementWiseUnaryOperatorType::OpName>::Create)
-
-	OP(DML_ELEMENT_WISE_ABS_OPERATOR_DESC, Abs);
-	OP(DML_ELEMENT_WISE_ACOS_OPERATOR_DESC, Acos);
-	OP(DML_ELEMENT_WISE_ACOSH_OPERATOR_DESC, Acosh);
-	OP(DML_ELEMENT_WISE_ASIN_OPERATOR_DESC, Asin);
-	OP(DML_ELEMENT_WISE_ASINH_OPERATOR_DESC, Asinh);
-	OP(DML_ELEMENT_WISE_ATAN_OPERATOR_DESC, Atan);
-	OP(DML_ELEMENT_WISE_ATANH_OPERATOR_DESC, Atanh);
-	OP(DML_ELEMENT_WISE_CEIL_OPERATOR_DESC, Ceil);
-	OP(DML_ELEMENT_WISE_CLIP_OPERATOR_DESC, Clip);
-	OP(DML_ELEMENT_WISE_COS_OPERATOR_DESC, Cos);
-	OP(DML_ELEMENT_WISE_COSH_OPERATOR_DESC, Cosh);
-	OP(DML_ACTIVATION_ELU_OPERATOR_DESC, Elu);
-	OP(DML_ELEMENT_WISE_ERF_OPERATOR_DESC, Erf);
-	OP(DML_ELEMENT_WISE_EXP_OPERATOR_DESC, Exp);
-	OP(DML_ELEMENT_WISE_FLOOR_OPERATOR_DESC, Floor);
-	OP(DML_ELEMENT_WISE_IS_INFINITY_OPERATOR_DESC, IsInf);
-	OP(DML_ELEMENT_WISE_IS_NAN_OPERATOR_DESC, IsNan);
-	OP(DML_ACTIVATION_HARD_SIGMOID_OPERATOR_DESC, HardSigmoid);
-	//OP(HardSwish);
-	OP(DML_ACTIVATION_LEAKY_RELU_OPERATOR_DESC, LeakyRelu);
-	OP(DML_ELEMENT_WISE_LOG_OPERATOR_DESC, Log);
-	OP(DML_ELEMENT_WISE_NEGATE_OPERATOR_DESC, Neg);
-	//OP(Not);
-	OP(DML_ELEMENT_WISE_RECIP_OPERATOR_DESC, Reciprocal);
-	OP(DML_ACTIVATION_RELU_OPERATOR_DESC, Relu);
-	OP(DML_ELEMENT_WISE_ROUND_OPERATOR_DESC, Round);
-	OP(DML_ACTIVATION_SCALED_ELU_OPERATOR_DESC, Selu);
-	OP(DML_ACTIVATION_SIGMOID_OPERATOR_DESC, Sigmoid);
-	OP(DML_ELEMENT_WISE_SIGN_OPERATOR_DESC, Sign);
-	OP(DML_ELEMENT_WISE_SIN_OPERATOR_DESC, Sin);
-	OP(DML_ELEMENT_WISE_SINH_OPERATOR_DESC, Sinh);
-	OP(DML_ACTIVATION_SOFTPLUS_OPERATOR_DESC, Softplus);
-	OP(DML_ACTIVATION_SOFTSIGN_OPERATOR_DESC, Softsign);
-	OP(DML_ELEMENT_WISE_SQRT_OPERATOR_DESC, Sqrt);
-	OP(DML_ELEMENT_WISE_TAN_OPERATOR_DESC, Tan);
-	OP(DML_ELEMENT_WISE_TANH_OPERATOR_DESC, Tanh);
-
-#undef OP
-}
-
-void RegisterElementWiseBinaryOperators()
-{
-#define OP(DmlOpDesc, OpName) FOperatorRegistryDml::Get()->OpAdd(TEXT(#OpName), FOperatorDmlElementWiseBinary<DmlOpDesc, NNECore::Internal::EElementWiseBinaryOperatorType::OpName>::Create)
-
-	OP(DML_ELEMENT_WISE_ADD_OPERATOR_DESC, Add);
-	// And
-	OP(DML_ELEMENT_WISE_DIVIDE_OPERATOR_DESC, Div);
-	// Equal
-	// Greater
-	// GreaterOrEqual
-	// Less
-	// LessOrEqual
-	// Mod
-	OP(DML_ELEMENT_WISE_MULTIPLY_OPERATOR_DESC, Mul);
-	// Or
-	OP(DML_ACTIVATION_PARAMETERIZED_RELU_OPERATOR_DESC, Prelu);
-	OP(DML_ELEMENT_WISE_POW_OPERATOR_DESC, Pow);
-	OP(DML_ELEMENT_WISE_SUBTRACT_OPERATOR_DESC, Sub);
-	// Xor
-
-#undef OP
-}
-
-struct FDmlOperatorElementWiseRegistrator
-{
-	FDmlOperatorElementWiseRegistrator()
-	{
-		RegisterElementWiseUnaryOperators();
-		RegisterElementWiseBinaryOperators();
-	}
-};
+#define REGISTER_OP_ELEMENT_WISE(OpName, DmlOpName, OpClass) \
+struct FDmlOperator##OpName##Registrator \
+{ \
+	FDmlOperator##OpName##Registrator() \
+	{ \
+		FOperatorRegistryDml::Get()->OpAdd(TEXT(#OpName), OpClass<DML_ELEMENT_WISE_##DmlOpName##_OPERATOR_DESC, DML_OPERATOR_ELEMENT_WISE_##DmlOpName>::Create); \
+	} \
+}; \
+\
+static FDmlOperator##OpName##Registrator RegisterDmlOperator##OpName;
+
+
+#define REGISTER_OP_ELEMENT_WISE_UNARY(OpName, DmlOpName) REGISTER_OP_ELEMENT_WISE(OpName, DmlOpName, FOperatorDmlElementWiseUnary)
+
+#define REGISTER_OP_ELEMENT_WISE_BINARY(OpName, DmlOpName) REGISTER_OP_ELEMENT_WISE(OpName, DmlOpName, FOperatorDmlElementWiseBinary)
+
+// 	REGISTER_OP_ELEMENT_WISE_UNARY(	OpName, 			DmlOpName )
+	REGISTER_OP_ELEMENT_WISE_UNARY( Abs,				ABS )
+	REGISTER_OP_ELEMENT_WISE_UNARY( Acos, 				ACOS )
+	REGISTER_OP_ELEMENT_WISE_UNARY( Acosh, 				ACOSH )
+	REGISTER_OP_ELEMENT_WISE_UNARY( Asin, 				ASIN )
+	REGISTER_OP_ELEMENT_WISE_UNARY( Asinh, 				ASINH )
+	REGISTER_OP_ELEMENT_WISE_UNARY( Atan,				ATAN )
+	REGISTER_OP_ELEMENT_WISE_UNARY( Atanh, 				ATANH )
+	REGISTER_OP_ELEMENT_WISE_UNARY( Ceil,				CEIL )
+	REGISTER_OP_ELEMENT_WISE_UNARY( Clip, 				CLIP )
+	REGISTER_OP_ELEMENT_WISE_UNARY( Cos,				COS )
+	REGISTER_OP_ELEMENT_WISE_UNARY( Cosh,				COSH )
+	REGISTER_OP_ELEMENT_WISE_UNARY( Erf,				ERF )
+	REGISTER_OP_ELEMENT_WISE_UNARY( Exp, 				EXP )
+	REGISTER_OP_ELEMENT_WISE_UNARY( Floor,				FLOOR )
+	REGISTER_OP_ELEMENT_WISE_UNARY( IsInf, 				IS_INFINITY )
+	REGISTER_OP_ELEMENT_WISE_UNARY( IsNan, 				IS_NAN )
+	REGISTER_OP_ELEMENT_WISE_UNARY( Log,				LOG )
+	REGISTER_OP_ELEMENT_WISE_UNARY( Neg,				NEGATE )
+	REGISTER_OP_ELEMENT_WISE_UNARY( Reciprocal, 		RECIP )
+	REGISTER_OP_ELEMENT_WISE_UNARY( Round,				ROUND )
+	REGISTER_OP_ELEMENT_WISE_UNARY( Sign,				SIGN )
+	REGISTER_OP_ELEMENT_WISE_UNARY( Sin,				SIN )
+	REGISTER_OP_ELEMENT_WISE_UNARY( Sinh,				SINH )
+	REGISTER_OP_ELEMENT_WISE_UNARY( Sqrt,				SQRT )
+	REGISTER_OP_ELEMENT_WISE_UNARY( Tan,				TAN )
+	REGISTER_OP_ELEMENT_WISE_UNARY( Tanh,				TANH )
+
+// 	REGISTER_OP_ELEMENT_WISE_BINARY( OpName, 			DmlOpName )
+	REGISTER_OP_ELEMENT_WISE_BINARY( Add,				ADD )
+	REGISTER_OP_ELEMENT_WISE_BINARY( Div,				DIVIDE )
+	REGISTER_OP_ELEMENT_WISE_BINARY( Mul,				MULTIPLY )
+	REGISTER_OP_ELEMENT_WISE_BINARY( Pow,				POW )
+	REGISTER_OP_ELEMENT_WISE_BINARY( Sub,				SUBTRACT )
+
 
-static FDmlOperatorElementWiseRegistrator RegisterElementWiseOperators;
 
 } // namespace UE::NNERuntimeRDG::Private::Dml
 
