Add RDG DML reduce operators
Enable batch normalization QA tests

#rnx

diff --git a/Engine/Plugins/Experimental/NNE/Source/NNEQA/Resources/NNEQATestDesc.json b/Engine/Plugins/Experimental/NNE/Source/NNEQA/Resources/NNEQATestDesc.json
--- a/Engine/Plugins/Experimental/NNE/Source/NNEQA/Resources/NNEQATestDesc.json
+++ b/Engine/Plugins/Experimental/NNE/Source/NNEQA/Resources/NNEQATestDesc.json
@@ -336,9 +336,6 @@
 		},
 		{
 			"category": "BatchNormalization",
-			"runtimes": [
-				{ "name": "NNERuntimeRDGDml", "skip": true },
-			],
 			"targets": [
 				{ "target": "BatchNormalization", "relative_tolerance": 5e-1 }
 			]
@@ -411,6 +408,24 @@
 				{ "target": "Concat" }
 			]
 		},
+		{
+			"category": "Reduce",
+			"runtimes": [
+				{ "name": "NNERuntimeRDGHlsl", "skip": true }
+			],
+			"targets": [
+				{ "target": "ReduceL1" },
+				{ "target": "ReduceL2" },
+				{ "target": "ReduceLogSum" },
+				{ "target": "ReduceLogSumExp" },
+				{ "target": "ReduceMin" },
+				{ "target": "ReduceMax" },
+				{ "target": "ReduceMean" },
+				{ "target": "ReduceProd" },
+				{ "target": "ReduceSum" },
+				{ "target": "ReduceSumSquare" }
+			]
+		},
 		{
 			"category": "Reshape",
 			"targets": [
@@ -1032,7 +1047,7 @@
 			]
 		},
 		{ 
-			"name": "Reshape.ReducedDims",
+			"name": "Reshape.RducedDims",
 			"datasets": [
 				{ "inputs": [ { "shape": [ 2, 3, 4 ] } ], "weights": [ { "shape": [ 2 ], "source": [ "2", "12" ], "type": "int64" } ], "outputs": [ { "shape": [ 2, 12 ] } ] }
 			]
@@ -1072,6 +1087,30 @@
 			"datasets": [
 				{ "inputs": [ { "shape": [ 1, 2, 256, 256 ] } ], "outputs": [ { "shape": [ 1, 2, 128, 128 ] } ] }
 			]
+		},
+		{ 
+			"name": "Reduce.Default",
+			"datasets": [
+				{ "inputs": [ { "shape": [ 3, 2, 2 ] } ], "outputs": [ { "shape": [ 1, 1, 1 ] } ] }
+			]
+		},
+		{ 
+			"name": "Reduce.Axes.Keepdims",
+			"datasets": [
+				{ "inputs": [ { "shape": [ 3, 2, 2 ] } ], "outputs": [ { "shape": [ 3, 1, 2 ] } ] }
+			]
+		},
+		{ 
+			"name": "Reduce.Axes.DontKeepdims",
+			"datasets": [
+				{ "inputs": [ { "shape": [ 3, 2, 2 ] } ], "outputs": [ { "shape": [ 3, 2 ] } ] }
+			]
+		},
+		{ 
+			"name": "Reduce.NegativeAxes.Keepdims",
+			"datasets": [
+				{ "inputs": [ { "shape": [ 3, 2, 2 ] } ], "outputs": [ { "shape": [ 3, 1, 2 ] } ] }
+			]
 		}
 	],
 
@@ -1553,6 +1592,37 @@
 					]
 				}
 			]
+		},
+		{
+			"name": "Reduce.Axes.Keepdims",
+			"attribute_maps": [
+				{
+					"attributes": [
+						{ "name": "axes", "value": [ 1 ], "type": "Int32Array" },
+						{ "name": "keepdims", "value": 1, "type": "Int32" }
+					]
+				}
+			]
+		},
+		{
+			"name": "Reduce.Axes.DontKeepdims",
+			"attribute_maps": [
+				{ "attributes": [ 
+						{ "name": "axes", "value": [1], "type": "Int32Array" },
+						{ "name": "keepdims", "value": 0, "type": "Int32" }
+					]
+				}
+			]
+		},
+		{
+			"name": "Reduce.NegativeAxes.Keepdims",
+			"attribute_maps": [
+				{ "attributes": [ 
+						{ "name": "axes", "value": [-2], "type": "Int32Array" },
+						{ "name": "keepdims", "value": 1, "type": "Int32" }
+					]
+				}
+			]
 		}
 	]
 
diff --git a/Engine/Plugins/Experimental/NNE/Source/NNERuntimeRDG/Private/Dml/NNEDmlOperatorReduce.cpp b/Engine/Plugins/Experimental/NNE/Source/NNERuntimeRDG/Private/Dml/NNEDmlOperatorReduce.cpp
new file mode 100644
--- /dev/null
+++ b/Engine/Plugins/Experimental/NNE/Source/NNERuntimeRDG/Private/Dml/NNEDmlOperatorReduce.cpp
@@ -0,0 +1,161 @@
+// Copyright Epic Games, Inc. All Rights Reserved.
+
+#ifdef NNE_USE_DIRECTML
+
+#include "NNEDmlOperator.h"
+#include "NNEDmlOperatorUtils.h"
+
+namespace UE::NNERuntimeRDG::Private::Dml
+{
+
+//
+// Reduce operator covers:
+// ReduceL1
+// ReduceL2
+// ReduceLogSum
+// ReduceLogSumExp
+// ReduceMin
+// ReduceMax
+// ReduceMean
+// ReduceProd
+// ReduceSum
+// ReduceSumSquare
+template<DML_REDUCE_FUNCTION ReduceFunc>
+class FOperatorDmlReduce : public FOperatorDml
+{
+
+	inline void HandleEmptyAxes(TArray<int32>& Axes, int32 Rank)
+	{
+		if (Axes.IsEmpty())
+		{
+			Axes.SetNumUninitialized(Rank);
+
+			for (int32 Idx = 0; Idx < Rank; ++Idx)
+			{
+				Axes[Idx] = Idx;
+			}
+		}
+	}
+
+public:
+
+	static FOperatorDml* Create()
+	{
+		return new FOperatorDmlReduce<ReduceFunc>();
+	}
+
+	//
+	//
+	//
+	virtual bool Initialize(IDMLDevice* Device, TArrayView<const NNECore::Internal::FTensor> InputTensors, TArrayView<const NNECore::Internal::FTensor> OutputTensors, const NNECore::FAttributeMap& Attributes) override
+	{
+		checkf(InputTensors.Num() == 1, TEXT("Dml Reduce op supports only 1 input"));
+		check(OutputTensors.Num() == 1);
+
+		const int32 KeepDims = Attributes.GetValueOrDefault<int32>(TEXT("keepdims"), 1);
+		
+		TArray<int32>	Axes;
+		
+		const FNNEAttributeValue* AxesAttr = Attributes.GetAttributeValue(TEXT("axes"));
+		if (AxesAttr)
+		{
+			Axes = AxesAttr->GetValue<TArray<int32>>();
+		}
+
+		HandleNegativeAxes(Axes, InputTensors[0].GetShape().Rank());
+		HandleEmptyAxes(Axes, InputTensors[0].GetShape().Rank());
+
+		DmlUtil::FTensorDesc DmlInputTensorDesc;
+		
+		if (!DmlInputTensorDesc.InitFromTensor(InputTensors[0], InputTensors[0].GetShape().Rank()))
+		{
+			UE_LOG(LogNNE, Warning, TEXT("Failed to initialize Reduce input tensor for DML inference"));
+			return false;
+		}
+
+		DmlUtil::FSmallUIntArray	DmlAxes;
+		DmlUtil::FSmallUIntArray	ReducedDims;
+
+		ReducedDims.SetNum(DmlInputTensorDesc.Sizes.Num());
+
+		for (int32& Dim : Axes)
+		{
+			checkf(Dim < ReducedDims.Num(), TEXT("Index out of bounds for Reduce axis"));
+			ReducedDims[Dim] = 1;
+			DmlAxes.Add(Dim);
+		}
+
+		DmlUtil::FSmallUIntArray	OutputShape;
+
+		if (!KeepDims)
+		{
+			// ReduceSum example:
+			//     input dims: {3, 2, 2}
+			//     axes: 1
+			//     keepDims: 0
+			// 
+			// ONNX: {3, 2}, but DML: {3, 1, 2}	
+			TConstArrayView<uint32> InputShape = InputTensors[0].GetShape().GetData();
+
+			OutputShape.SetNum(ReducedDims.Num());
+
+			for (int32 Idx = 0; Idx < OutputShape.Num(); ++Idx)
+			{
+				OutputShape[Idx] = ReducedDims[Idx] ? 1 : InputShape[Idx];
+			}
+		}
+		else
+		{
+			OutputShape.Append(OutputTensors[0].GetShape().GetData());
+		}
+
+		DmlUtil::FTensorDesc DmlOutputTensorDesc;
+
+		if (!DmlOutputTensorDesc.InitFromTensor(OutputTensors[0], OutputShape.Num(), 
+			/*Broadcast shape*/ MakeEmptyArrayView<uint32>(),
+			/*Custom shape*/ OutputShape))
+		{
+			UE_LOG(LogNNE, Warning, TEXT("Failed to initialize Reduce output tensor for DML inference"));
+			return false;
+		}
+
+		DML_REDUCE_OPERATOR_DESC OpDesc{};
+
+		OpDesc.InputTensor = &DmlInputTensorDesc.Desc;
+		OpDesc.OutputTensor = &DmlOutputTensorDesc.Desc;
+		OpDesc.Function = ReduceFunc;
+		OpDesc.Axes = DmlAxes.GetData();
+		OpDesc.AxisCount = (uint32) DmlAxes.Num();
+
+		return CreateOperator(Device, DML_OPERATOR_DESC{ DML_OPERATOR_REDUCE, &OpDesc} );
+	}
+};
+
+// Register Reshape operator on Module startup
+#define OP(OpName, ReduceFunc) FOperatorRegistryDml::Get()->OpAdd(TEXT(#OpName), FOperatorDmlReduce<##ReduceFunc##>::Create)
+
+struct FOperatorDmlReduceRegistrator
+{
+	FOperatorDmlReduceRegistrator()
+	{
+		OP(ReduceL1,		DML_REDUCE_FUNCTION_L1);
+		OP(ReduceL2,		DML_REDUCE_FUNCTION_L2);
+		OP(ReduceLogSum,	DML_REDUCE_FUNCTION_LOG_SUM);
+		OP(ReduceLogSumExp,	DML_REDUCE_FUNCTION_LOG_SUM_EXP);
+		OP(ReduceMin,		DML_REDUCE_FUNCTION_MIN);
+		OP(ReduceMax,		DML_REDUCE_FUNCTION_MAX);
+		OP(ReduceMean,		DML_REDUCE_FUNCTION_AVERAGE);
+		OP(ReduceProd,		DML_REDUCE_FUNCTION_MULTIPLY);
+		OP(ReduceSum,		DML_REDUCE_FUNCTION_SUM);
+		OP(ReduceSumSquare,	DML_REDUCE_FUNCTION_SUM_SQUARE);
+	}
+};
+
+#undef OP
+
+static FOperatorDmlReduceRegistrator RegisterReduceOperators;
+
+
+} // namespace UE::NNERuntimeRDG::Private::Dml
+
+#endif // NNE_USE_DIRECTML
diff --git a/Engine/Plugins/Experimental/NNE/Source/NNERuntimeRDG/Private/Dml/NNEDmlOperatorUtils.h b/Engine/Plugins/Experimental/NNE/Source/NNERuntimeRDG/Private/Dml/NNEDmlOperatorUtils.h
--- a/Engine/Plugins/Experimental/NNE/Source/NNERuntimeRDG/Private/Dml/NNEDmlOperatorUtils.h
+++ b/Engine/Plugins/Experimental/NNE/Source/NNERuntimeRDG/Private/Dml/NNEDmlOperatorUtils.h
@@ -126,6 +126,14 @@ inline int32 HandleNegativeAxis(int32 Axis, int32 Rank)
 	return Axis;
 }
 
+inline void HandleNegativeAxes(TArrayView<int32> Axes, int32 Rank)
+{
+	for (int32& Axis : Axes)
+	{
+		Axis = HandleNegativeAxis(Axis, Rank);
+	}
+}
+
 inline int32 GetDmlAxis(int32 OnnxAxis, int32 OnnxDim, int32 DmlDim)
 {
 	check(DmlDim >= OnnxDim);
diff --git a/Engine/Plugins/Experimental/NNE/Source/NNEUtils/Private/NNEUtilsModelBuilderONNX.cpp b/Engine/Plugins/Experimental/NNE/Source/NNEUtils/Private/NNEUtilsModelBuilderONNX.cpp
--- a/Engine/Plugins/Experimental/NNE/Source/NNEUtils/Private/NNEUtilsModelBuilderONNX.cpp
+++ b/Engine/Plugins/Experimental/NNE/Source/NNEUtils/Private/NNEUtilsModelBuilderONNX.cpp
@@ -357,10 +357,27 @@ NNEUTILS_API bool CreateONNXModelForOperator(bool UseVariadicShapeForModel, cons
 		OperatorName == TEXT("Slice") ||				// current implementation is opset 1 (next version is 10)
 		OperatorName == TEXT("Squeeze") ||				// current implementation is opset 1 (next version is 11)
 		OperatorName == TEXT("Unsqueeze") ||			// current implementation is opset 1 (next version is 11)
-		OperatorName == TEXT("Upsample"))				// deprecated starting opset 10
+		OperatorName == TEXT("Upsample")				// deprecated starting opset 10
+		)				
 	{
 		OpsetVersion = 9;
 	}
+	else
+	if (OperatorName == TEXT("ReduceL1") ||				// current implementation is opset 11 (next version is 13)
+		OperatorName == TEXT("ReduceL2") ||				// current implementation is opset 11 (next version is 13)
+		OperatorName == TEXT("ReduceLogSum") ||			// current implementation is opset 11 (next version is 13)
+		OperatorName == TEXT("ReduceLogSumExp") ||		// current implementation is opset 11 (next version is 13)
+		OperatorName == TEXT("ReduceLogMin") ||			// current implementation is opset 11 (next version is 13)
+		OperatorName == TEXT("ReduceLogMax") ||			// current implementation is opset 11 (next version is 13)
+		OperatorName == TEXT("ReduceMean") ||			// current implementation is opset 11 (next version is 13)
+		OperatorName == TEXT("ReduceProd") ||			// current implementation is opset 11 (next version is 13)
+		OperatorName == TEXT("ReduceSum") ||			// current implementation is opset 11 (next version is 13)
+		OperatorName == TEXT("ReduceSumSquare")			// current implementation is opset 11 (next version is 13)
+		)
+	{
+		OpsetVersion = 11;
+	}
+
 	TUniquePtr<IModelBuilder> Builder(CreateONNXModelBuilder(IrVersion, OpsetVersion));
 
 	Builder->Begin();
