Add RDG DirectML Squeeze and Unsqueeze operator

#rnx

diff --git a/Engine/Plugins/Experimental/NNE/Source/NNEQA/Resources/NNEQATestDesc.json b/Engine/Plugins/Experimental/NNE/Source/NNEQA/Resources/NNEQATestDesc.json
--- a/Engine/Plugins/Experimental/NNE/Source/NNEQA/Resources/NNEQATestDesc.json
+++ b/Engine/Plugins/Experimental/NNE/Source/NNEQA/Resources/NNEQATestDesc.json
@@ -438,18 +438,12 @@
 		},
 		{
 			"category": "SqueezeOp",
-			"runtimes": [
-				{ "name": "NNERuntimeRDGDml", "skip": true },
-			],
 			"targets": [
 				{ "target": "Squeeze" }
 			]
 		},
 		{
 			"category": "Unsqueeze",
-			"runtimes": [
-				{ "name": "NNERuntimeRDGDml", "skip": true },
-			],
 			"targets": [
 				{ "target": "Unsqueeze" }
 			]
@@ -1035,17 +1029,29 @@
 			]
 		},
 		{ 
-			"name": "Unsqueeze",
+			"name": "Unsqueeze.Unsqueeze1",
 			"datasets": [
 				{ "inputs": [ { "shape": [ 2, 3, 4 ] } ], "outputs": [ { "shape": [ 2, 1, 3, 4 ] } ] }
 			]
 		},
 		{ 
-			"name": "SqueezeOp",
+			"name": "Unsqueeze.Unsqueeze2",
+			"datasets": [
+				{ "inputs": [ { "shape": [ 1, 131 ] } ], "outputs": [ { "shape": [ 1, 131, 1 ] } ] }
+			]
+		},
+		{ 
+			"name": "SqueezeOp.Axes",
 			"datasets": [
 				{ "inputs": [ { "shape": [ 2, 3, 1, 4 ] } ], "outputs": [ { "shape": [ 2, 3, 4 ] } ] }
 			]
 		},
+		{ 
+			"name": "SqueezeOp.MultiAxes",
+			"datasets": [
+				{ "inputs": [ { "shape": [ 2, 1, 3, 1, 4 ] } ], "outputs": [ { "shape": [ 2, 3, 4 ] } ] }
+			]
+		},
 		{ 
 			"name": "Reshape.Reshape234to423",
 			"datasets": [
@@ -1560,7 +1566,7 @@
 			]
 		},
 		{
-			"name": "SqueezeOp",
+			"name": "SqueezeOp.Axes",
 			"attribute_maps": [
 				{ "attributes": [
 					{ "name": "axes", "value": [2], "type": "Int32Array" }
@@ -1568,13 +1574,29 @@
 			]
 		},
 		{
-			"name": "Unsqueeze",
+			"name": "SqueezeOp.MultiAxes",
+			"attribute_maps": [
+				{ "attributes": [
+					{ "name": "axes", "value": [1, 3], "type": "Int32Array" }
+				] }
+			]
+		},
+		{
+			"name": "Unsqueeze.Unsqueeze1",
 			"attribute_maps": [
 				{ "attributes": [
 					{ "name": "axes", "value": [1], "type": "Int32Array" }
 				] }
 			]
 		},
+		{
+			"name": "Unsqueeze.Unsqueeze2",
+			"attribute_maps": [
+				{ "attributes": [
+					{ "name": "axes", "value": [2], "type": "Int32Array" }
+				] }
+			]
+		},
 		{
 			"name": "Concat.Concat1",
 			"attribute_maps": [
diff --git a/Engine/Plugins/Experimental/NNE/Source/NNERuntimeRDG/Private/Dml/NNEDmlOperatorUnsqueeze.cpp b/Engine/Plugins/Experimental/NNE/Source/NNERuntimeRDG/Private/Dml/NNEDmlOperatorUnsqueeze.cpp
new file mode 100644
--- /dev/null
+++ b/Engine/Plugins/Experimental/NNE/Source/NNERuntimeRDG/Private/Dml/NNEDmlOperatorUnsqueeze.cpp
@@ -0,0 +1,83 @@
+// Copyright Epic Games, Inc. All Rights Reserved.
+
+#ifdef NNE_USE_DIRECTML
+
+#include "NNEDmlOperator.h"
+#include "NNEDmlOperatorUtils.h"
+
+namespace UE::NNERuntimeRDG::Private::Dml
+{
+
+class FOperatorDmlUnsqueeze : public FOperatorDml
+{
+
+public:
+
+	static FOperatorDml* Create()
+	{
+		return new FOperatorDmlUnsqueeze();
+	}
+
+	//
+	//
+	//
+	virtual bool Initialize(IDMLDevice* Device, TArrayView<const NNECore::Internal::FTensor> InputTensors, TArrayView<const NNECore::Internal::FTensor> OutputTensors, const NNECore::FAttributeMap& Attributes) override
+	{
+		check(InputTensors.Num() == 1);
+		check(OutputTensors.Num() == 1);
+
+		const FNNEAttributeValue* AxesAttr = Attributes.GetAttributeValue(TEXT("axes"));
+		if (!AxesAttr)
+		{
+			UE_LOG(LogNNE, Warning, TEXT("Missing axes attribute for operator Unsqueeze"));
+			return false;
+		}
+
+		TArray<int32>			Axes = AxesAttr->GetValue<TArray<int32>>();
+		TConstArrayView<uint32>	InputShape = InputTensors[0].GetShape().GetData();
+		const int32_t			OutShapeRank = InputShape.Num() + Axes.Num();
+
+		HandleNegativeAxes(Axes, OutShapeRank);
+		Axes.Sort();
+
+		DmlUtil::FSmallUIntArray	OutputShape;
+
+		OutputShape.Reserve(OutShapeRank);
+		OutputShape.Append(InputShape);
+		
+		for (int32 Idx = 0; Idx < Axes.Num(); ++Idx)
+		{
+			if (Axes[Idx] >= OutShapeRank)
+			{
+				UE_LOG(LogNNE, Warning, TEXT("Unsqueeze operator does not support axes greater than the number of dimensions of the resulting tensor shape"));
+				return false;
+			}
+
+			OutputShape.Insert(1, Axes[Idx]);
+		}
+
+		DmlUtil::FTensorDesc DmlTensorDesc;
+		
+		if (!DmlTensorDesc.InitFromTensor(OutputTensors[0], OutputShape.Num(),
+			/*Broadcast =*/ MakeEmptyArrayView<uint32>(),
+			/*CustomShape =*/ OutputShape))
+		{
+			UE_LOG(LogNNE, Warning, TEXT("Failed to initialize Unsqueeze's output tensor for DML inference"));
+			return false;
+		}
+
+		DML_ELEMENT_WISE_IDENTITY_OPERATOR_DESC DmlIdentityOpDesc{};
+
+		DmlIdentityOpDesc.InputTensor = &DmlTensorDesc.Desc;
+		DmlIdentityOpDesc.OutputTensor = &DmlTensorDesc.Desc;
+
+		return CreateOperator(Device, DML_OPERATOR_DESC{ DML_OPERATOR_ELEMENT_WISE_IDENTITY, &DmlIdentityOpDesc });
+	}
+};
+
+// Register Reshape operator on Module startup
+NNE_DML_REGISTER_OP(Unsqueeze)
+
+} // namespace UE::NNERuntimeRDG::Private::Dml
+
+#endif // NNE_USE_DIRECTML
diff --git a/Engine/Plugins/Experimental/NNE/Source/NNERuntimeRDG/Private/Hlsl/NNERuntimeRDGSqueeze.cpp b/Engine/Plugins/Experimental/NNE/Source/NNERuntimeRDG/Private/Hlsl/NNERuntimeRDGSqueeze.cpp
--- a/Engine/Plugins/Experimental/NNE/Source/NNERuntimeRDG/Private/Hlsl/NNERuntimeRDGSqueeze.cpp
+++ b/Engine/Plugins/Experimental/NNE/Source/NNERuntimeRDG/Private/Hlsl/NNERuntimeRDGSqueeze.cpp
@@ -31,8 +31,10 @@ namespace UE::NNERuntimeRDG::Private::Hlsl
 			const NNECore::Internal::FTensor& X = *InputTensors[0];
 			TArray<uint32> OutputShapeData(X.GetShape().GetData());
 
-			for (int32 Axe : Axes)
+			for (int32 Idx = Axes.Num() - 1; Idx >= 0; --Idx)
 			{
+				const int32 Axe = Axes[Idx];
+
 				if (OutputShapeData[Axe] != 1)
 				{
 					UE_LOG(LogNNE, Warning, TEXT("Squeeze at axe %d for 'Data' (name: %s) should be targeting a dimension of size 1 but it is %d."), Axe, *X.GetName(), OutputShapeData[Axe]);
diff --git a/Engine/Plugins/Experimental/NNE/Source/NNEUtils/Private/NNEUtilsModelBuilderONNX.cpp b/Engine/Plugins/Experimental/NNE/Source/NNEUtils/Private/NNEUtilsModelBuilderONNX.cpp
--- a/Engine/Plugins/Experimental/NNE/Source/NNEUtils/Private/NNEUtilsModelBuilderONNX.cpp
+++ b/Engine/Plugins/Experimental/NNE/Source/NNEUtils/Private/NNEUtilsModelBuilderONNX.cpp
@@ -350,6 +350,7 @@ NNEUTILS_API bool CreateONNXModelForOperator(bool UseVariadicShapeForModel, cons
 	
 	int64 IrVersion = OnnxIrVersion;
 	int64 OpsetVersion = OnnxOpsetVersion;
+
 	if (OperatorName == TEXT("BatchNormalization") ||	// current implementation is opset 9 (next version is 14)
 		OperatorName == TEXT("Clip") ||					// current implementation is opset 6 (next version is 11)
 		OperatorName == TEXT("Pad") ||					// current implementation is opset 2 (next version is 11)
@@ -372,7 +373,9 @@ NNEUTILS_API bool CreateONNXModelForOperator(bool UseVariadicShapeForModel, cons
 		OperatorName == TEXT("ReduceMean") ||			// current implementation is opset 11 (next version is 13)
 		OperatorName == TEXT("ReduceProd") ||			// current implementation is opset 11 (next version is 13)
 		OperatorName == TEXT("ReduceSum") ||			// current implementation is opset 11 (next version is 13)
-		OperatorName == TEXT("ReduceSumSquare")			// current implementation is opset 11 (next version is 13)
+		OperatorName == TEXT("ReduceSumSquare") ||		// current implementation is opset 11 (next version is 13)
+		OperatorName == TEXT("Squeeze") ||				// current implementation is opset 11 (next version is 13)
+		OperatorName == TEXT("Unsqueeze")				// current implementation is opset 11 (next version is 13)
 		)
 	{
 		OpsetVersion = 11;
