NNE RDGDml Fixes + Resize operator (partial support)

* Added Resize operator (restricted to DML 5 definition).
* Support for 0 sized cpu only tensors.
* Fix for pad v15 using const tensor for scale.
* Fix for network forwarding an input to more than on op.

Tests:
* All nne test green on editor win64
* Enabled more model tests for RDGDml + added a test for input forwarded to more than one node.

#rnx

diff --git a/Engine/Plugins/Experimental/NNE/Source/NNEQA/Resources/NNEQATestDesc.json b/Engine/Plugins/Experimental/NNE/Source/NNEQA/Resources/NNEQATestDesc.json
--- a/Engine/Plugins/Experimental/NNE/Source/NNEQA/Resources/NNEQATestDesc.json
+++ b/Engine/Plugins/Experimental/NNE/Source/NNEQA/Resources/NNEQATestDesc.json
@@ -10,8 +10,8 @@
 					"absolute_tolerance": 5.0e-2
 				}
 			]
-				},
-				{
+		},
+		{
 			"category": "StyleTransfer.VariableInputSize",
 			"is_model_test": true,
 			"targets": [
@@ -67,17 +67,15 @@
 			"is_model_test": true,
 			"targets": [
 				{
-					"target": "OneLayerWithWeights",
-					"runtimes": [
-						{ "name": "NNERuntimeRDGDml",  "skip": true, "comment": "Model not supported yet." }
-					]
+					"target": "OneLayerWithWeights"
 				},
 				{
 					"target": "TwoLayersNoWeight",
-					"tags": [ "SmokeTest" ],
-					"runtimes": [
-						{ "name": "NNERuntimeRDGDml",  "skip": true, "comment": "Model not supported yet." }
-					]
+					"tags": [ "SmokeTest" ]
+				},
+				{
+					"target": "MultiInput",
+					"tags": [ "SmokeTest" ]
 				},
 				{
 					"target": "MultiOutput",
diff --git a/Engine/Plugins/Experimental/NNE/Source/NNERuntimeRDG/Private/Dml/NNEDmlModel.cpp b/Engine/Plugins/Experimental/NNE/Source/NNERuntimeRDG/Private/Dml/NNEDmlModel.cpp
--- a/Engine/Plugins/Experimental/NNE/Source/NNERuntimeRDG/Private/Dml/NNEDmlModel.cpp
+++ b/Engine/Plugins/Experimental/NNE/Source/NNERuntimeRDG/Private/Dml/NNEDmlModel.cpp
@@ -493,7 +493,7 @@ public:
 			
 		DML_GRAPH_DESC	Graph = DML_GRAPH_DESC{};
 
-		Graph.InputCount = InputEdges.Num();
+		Graph.InputCount = NumInputs;
 		Graph.OutputCount = OutputEdges.Num();
 		Graph.NodeCount = Operators.Num();
 		Graph.Nodes = Nodes.GetData();
@@ -624,6 +624,11 @@ private:
 			{
 				AddInputEdge(TensorIdx);
 			}
+
+			if (InGraph.ConstantCPUIndices.Find(TensorIdx) == INDEX_NONE)
+			{
+				++NumInputs;
+			}
 		}
 
 		return true;
@@ -728,8 +733,6 @@ private:
 				.SetTensorIdx(TensorIdx)
 				.SetNodeSrcOutput(NumInputs)
 		);
-
-		++NumInputs;
 	}
 
 	void AddOutputEdge(int32 TensorIdx)
diff --git a/Engine/Plugins/Experimental/NNE/Source/NNERuntimeRDG/Private/Dml/NNEDmlOperatorPad.cpp b/Engine/Plugins/Experimental/NNE/Source/NNERuntimeRDG/Private/Dml/NNEDmlOperatorPad.cpp
--- a/Engine/Plugins/Experimental/NNE/Source/NNERuntimeRDG/Private/Dml/NNEDmlOperatorPad.cpp
+++ b/Engine/Plugins/Experimental/NNE/Source/NNERuntimeRDG/Private/Dml/NNEDmlOperatorPad.cpp
@@ -63,6 +63,7 @@ public:
 				return false;
 			}
 			Pads.Append(InputTensors[1].GetPreparedData<int64>());
+			ConstantCPUInputs.Add(1);
 		}
 		else
 		{
diff --git a/Engine/Plugins/Experimental/NNE/Source/NNERuntimeRDG/Private/Dml/NNEDmlOperatorResample.cpp b/Engine/Plugins/Experimental/NNE/Source/NNERuntimeRDG/Private/Dml/NNEDmlOperatorResample.cpp
new file mode 100644
--- /dev/null
+++ b/Engine/Plugins/Experimental/NNE/Source/NNERuntimeRDG/Private/Dml/NNEDmlOperatorResample.cpp
@@ -0,0 +1,185 @@
+// Copyright Epic Games, Inc. All Rights Reserved.
+
+#ifdef NNE_USE_DIRECTML
+#include "NNEDmlOperator.h"
+
+namespace UE::NNERuntimeRDG::Private::Dml
+{
+
+// Remove array entries of the given indices (in ascending order), shifting them toward the front.
+// There is a special check to avoid removing all the values, since returning a completely
+// empty array would frequently causes errors later in many uses (such as with dimensions).
+//
+// e.g. input values = {2,1,3,1,1,5}
+//      ellidable input indices = {1,3,4}
+//      output values = {2,3,5}
+template< typename TData, typename TAllocator >
+void RemoveValuesByIndex(TConstArrayView<uint32> Indices, TArray<TData, TAllocator>& Values, bool bKeepOneValue)
+{
+	// Keep the last value at least, if all values would otherwise be removed.
+	if (bKeepOneValue && !Indices.IsEmpty() && Indices.Num() == Values.Num()) 
+	{
+		Indices = Indices.RightChop(1);
+	}
+
+	for (int32 Idx = Indices.Num() - 1; Idx >= 0; --Idx)
+	{
+		Values.RemoveAt(Indices[Idx]);
+	}
+}
+
+// Upsample and Resize operator are implemented as a DML Resample operator
+template < bool IsResize>
+class FOperatorDmlResample : public FOperatorDml
+{
+	static DML_INTERPOLATION_MODE ModeFromString(FStringView StringVal)
+	{
+		if (FCString::Stricmp(StringVal.GetData(), TEXT("NEAREST")) == 0)
+		{
+			return DML_INTERPOLATION_MODE_NEAREST_NEIGHBOR;
+		}
+		else if (FCString::Stricmp(StringVal.GetData(), TEXT("LINEAR")) == 0)
+		{
+			return DML_INTERPOLATION_MODE_LINEAR;
+		}
+		else
+		{
+			UE_LOG(LogNNE, Warning, TEXT("Unsupported interpolation mode:%s, using nearest neighbor instead"), StringVal.GetData());
+			return DML_INTERPOLATION_MODE_NEAREST_NEIGHBOR;
+		}
+	}
+
+public:
+
+	static FOperatorDml* Create()
+	{
+		return new FOperatorDmlResample<IsResize>();
+	}
+
+	//
+	//
+	//
+	virtual bool Initialize(IDMLDevice* Device, TArrayView<const NNECore::Internal::FTensor> InputTensors, TArrayView<const NNECore::Internal::FTensor> OutputTensors, const NNECore::FAttributeMap& Attributes) override
+	{
+		check(InputTensors.Num() >= 1 && InputTensors.Num() < 4);
+		check(OutputTensors.Num() == 1);
+
+		const NNECore::Internal::FTensor& InputTensor = InputTensors[0];
+		const NNECore::Internal::FTensor& ScaleTensor = (InputTensors.Num() == 2)?InputTensors[1]:InputTensors[2];//Upsample have scale at position 1, while Resize at 2.
+		const NNECore::Internal::FTensor& OutputTensor = OutputTensors[0];
+
+		if (InputTensors.Num() == 3 && InputTensors[2].GetDataType() == ENNETensorDataType::Int64)
+		{
+			UE_LOG(LogNNE, Warning, TEXT("Unsupported input type for 'sizes' of name %s, only 'scales' of type float is supported."), *InputTensors[2].GetName());
+			return false;
+		}
+
+		if (!ScaleTensor.HasPreparedData())
+		{
+			UE_LOG(LogNNE, Warning, TEXT("scale should be a constant tensor, it is here a variable tensor of name %s."), *ScaleTensor.GetName());
+			return false;
+		}
+		
+		// Read attributes
+		DML_INTERPOLATION_MODE Mode = ModeFromString(Attributes.GetValue<FString>(TEXT("mode")));
+
+		if constexpr (IsResize)
+		{
+			if (Mode == DML_INTERPOLATION_MODE_NEAREST_NEIGHBOR)
+			{
+				FString NeareastMode = Attributes.GetValueOrDefault<FString>(TEXT("nearest_mode"), TEXT("round_prefer_floor"));
+				if (FCString::Stricmp(*NeareastMode, TEXT("floor")) != 0)
+				{
+					UE_LOG(LogNNE, Warning, TEXT("Unsupported neareast mode:%s, using floor instead"), *NeareastMode);
+				}
+			}
+
+			FString CoordinateTransformationMode = Attributes.GetValueOrDefault<FString>(TEXT("coordinate_transformation_mode"), TEXT("half_pixel"));
+			if (FCString::Stricmp(*CoordinateTransformationMode, TEXT("half_pixel")) != 0)
+			{
+				UE_LOG(LogNNE, Warning, TEXT("Unsupported coordinate transformation mode:%s, using half_pixel instead"), *CoordinateTransformationMode);
+			}
+		}
+
+		// Initialize tensor descriptors
+		DmlUtil::FTensorDesc	DmlInputTensor{};
+		DmlUtil::FTensorDesc	DmlOutputTensor{};
+
+		if (!DmlInputTensor.InitFromTensor(InputTensor, 1))
+		{
+			UE_LOG(LogNNE, Warning, TEXT("Failed to initialize tensor(s) for DML inference"));
+			return false;
+		}
+
+		if (!DmlOutputTensor.InitFromTensor(OutputTensor, 1))
+		{
+			UE_LOG(LogNNE, Warning, TEXT("Failed to initialize tensor(s) for DML inference"));
+			return false;
+		}
+
+		// Notify DML graph that we have constant CPU tensor inputs
+		for (int i = 1; i < InputTensors.Num(); ++i)
+		{
+			ConstantCPUInputs.Add(i);
+		}
+
+		TConstArrayView<float> Scales = ScaleTensor.GetPreparedData<float>();
+		
+		// Find any useless dimensions of size 1 that occur in both input and output
+		DmlUtil::FSmallUIntArray	SqueezeInds;
+		TConstArrayView<uint32>		InputShape = InputTensor.GetShape().GetData();
+		TConstArrayView<uint32>		OutputShape = OutputTensor.GetShape().GetData();
+
+		for (int32 Idx = 0, Rank = OutputTensor.GetShape().Rank(); Idx < Rank; ++Idx)
+		{
+			if (InputShape[Idx] == 1 && OutputShape[Idx] == 1)
+			{
+				SqueezeInds.Emplace(Idx);
+			}
+		}
+
+		if (!SqueezeInds.IsEmpty())
+		{
+			DmlUtil::FSmallUIntArray	SqueezedInputShape(InputShape.GetData(), InputShape.Num());
+			DmlUtil::FSmallUIntArray	SqueezedOutputShape(OutputShape.GetData(), OutputShape.Num());
+			DmlUtil::FSmallArray<float> ScaleValues(Scales.GetData(), Scales.Num());
+
+			RemoveValuesByIndex(SqueezeInds, SqueezedInputShape, true);
+			RemoveValuesByIndex(SqueezeInds, SqueezedOutputShape, true);
+			RemoveValuesByIndex(SqueezeInds, ScaleValues, true);
+
+			DmlInputTensor.UpdateShapeAndStrides(SqueezedInputShape);
+			DmlOutputTensor.UpdateShapeAndStrides(SqueezedOutputShape);
+			Scales = ScaleValues;
+		}
+
+		DML_RESAMPLE_OPERATOR_DESC	OpDesc{};
+
+		OpDesc.InputTensor = &DmlInputTensor.Desc;
+		OpDesc.OutputTensor = &DmlOutputTensor.Desc;
+		OpDesc.InterpolationMode = Mode;
+		OpDesc.ScaleCount = Scales.Num();
+		OpDesc.Scales = Scales.GetData();
+
+		return CreateOperator(Device, DML_OPERATOR_DESC { DML_OPERATOR_RESAMPLE, &OpDesc });
+	}
+};
+
+#define NNE_DML_REGISTER_RESAMPLE_OP(OpName, IsResize) \
+struct FDmlOperator##OpName##Registrator \
+{ \
+	FDmlOperator##OpName##Registrator() \
+	{ \
+		FOperatorRegistryDml::Get()->OpAdd(TEXT(#OpName), FOperatorDmlResample<IsResize>::Create); \
+	} \
+}; \
+\
+static FDmlOperator##OpName##Registrator RegisterDmlOperator##OpName;
+
+// Register operator on Module startup
+NNE_DML_REGISTER_RESAMPLE_OP(Upsample, false)
+NNE_DML_REGISTER_RESAMPLE_OP(Resize, true)
+
+} // namespace UE::NNERuntimeRDG::Private::Dml
+
+#endif // NNE_USE_DIRECTML
diff --git a/Engine/Plugins/Experimental/NNE/Source/NNERuntimeRDG/Private/Dml/NNEDmlOperatorUpsample.cpp b/Engine/Plugins/Experimental/NNE/Source/NNERuntimeRDG/Private/Dml/NNEDmlOperatorUpsample.cpp
deleted file mode 100644
diff --git a/Engine/Plugins/Experimental/NNE/Source/NNEUtils/Private/NNEUtilsModelBuilderNNE.cpp b/Engine/Plugins/Experimental/NNE/Source/NNEUtils/Private/NNEUtilsModelBuilderNNE.cpp
--- a/Engine/Plugins/Experimental/NNE/Source/NNEUtils/Private/NNEUtilsModelBuilderNNE.cpp
+++ b/Engine/Plugins/Experimental/NNE/Source/NNEUtils/Private/NNEUtilsModelBuilderNNE.cpp
@@ -113,8 +113,8 @@ public:
 		TArray<int32, TInlineAllocator<NNECore::FTensorShape::MaxRank>> NNEShape;
 		for (int i = 0; i < Shape.Num(); ++i)
 		{
-			//ORT Graph return 0 for variable dimensions, NNE use -1.
-			NNEShape.Emplace(Shape[i] == 0 ? -1 : Shape[i]);
+			//ORT Graph return 0 for variable dimensions for inputs/outputs, NNE use -1.
+			NNEShape.Emplace(!Data && Shape[i] == 0 ? -1 : Shape[i]);
 		}
 		
 		int Idx = AddTensor(Name, NNEShape, DataType, Data, DataSize);
diff --git a/Engine/Plugins/Experimental/NNE/Source/NNEUtils/Private/NNEUtilsModelOptimizer.cpp b/Engine/Plugins/Experimental/NNE/Source/NNEUtils/Private/NNEUtilsModelOptimizer.cpp
--- a/Engine/Plugins/Experimental/NNE/Source/NNEUtils/Private/NNEUtilsModelOptimizer.cpp
+++ b/Engine/Plugins/Experimental/NNE/Source/NNEUtils/Private/NNEUtilsModelOptimizer.cpp
@@ -630,16 +630,13 @@ private:
 				{
 					DataSize = Graph->GetTensorDataSize(TensorInit);
 
-					if (DataSize)
+					if (TensorDataBlob.Num() < DataSize)
 					{
-						if (TensorDataBlob.Num() < DataSize)
-						{
-							TensorDataBlob.SetNumUninitialized(DataSize);
-						}
-
-						Graph->GetTensorData(TensorInit, TensorDataBlob.GetData(), DataSize, 0);
-						Data = TensorDataBlob.GetData();
+						TensorDataBlob.SetNumUninitialized(DataSize);
 					}
+
+					Graph->GetTensorData(TensorInit, TensorDataBlob.GetData(), DataSize, 0);
+					Data = TensorDataBlob.GetData();
 				}
 
 				auto Tensor =
