Fixed DML binding issues

#rnx

diff --git a/Engine/Plugins/Experimental/NNE/Source/NNEQA/Resources/NNEQATestDesc.json b/Engine/Plugins/Experimental/NNE/Source/NNEQA/Resources/NNEQATestDesc.json
--- a/Engine/Plugins/Experimental/NNE/Source/NNEQA/Resources/NNEQATestDesc.json
+++ b/Engine/Plugins/Experimental/NNE/Source/NNEQA/Resources/NNEQATestDesc.json
@@ -59,6 +59,13 @@
 				{
 					"target": "VertexDeltaModel",
 					"absolute_tolerance": 1e-4
+				},
+				{
+					"target": "FaceDetector",
+					"runtimes": [
+						{ "name": "NNERuntimeRDGHlsl", "skip": true, "comment": "Model not supported yet." }
+					],
+					"absolute_tolerance": 1e-4
 				}
 			]
 		},
diff --git a/Engine/Plugins/Experimental/NNE/Source/NNERuntimeRDG/Private/Dml/NNEDmlModel.cpp b/Engine/Plugins/Experimental/NNE/Source/NNERuntimeRDG/Private/Dml/NNEDmlModel.cpp
--- a/Engine/Plugins/Experimental/NNE/Source/NNERuntimeRDG/Private/Dml/NNEDmlModel.cpp
+++ b/Engine/Plugins/Experimental/NNE/Source/NNERuntimeRDG/Private/Dml/NNEDmlModel.cpp
@@ -280,7 +280,7 @@ private:
 		else
 		{
 			DML_BUFFER_BINDING& Bind = Bindings.Add_GetRef({});
-			Descs.Add({ DML_BINDING_TYPE_BUFFER, &Bind });
+			Descs.Add({ DML_BINDING_TYPE_NONE, nullptr });
 		}
 	}
 
@@ -1168,8 +1168,9 @@ bool FModel::InitCompiledOp(TConstArrayView<int32> OpInputIndices, uint64 Tensor
 	return true;
 }
 
-BEGIN_SHADER_PARAMETER_STRUCT(FTensorBufferParamsDml, )
-	RDG_BUFFER_ACCESS(Buffer, ERHIAccess::UAVCompute)
+BEGIN_SHADER_PARAMETER_STRUCT(FDmlModelDispatchPassParameters, )
+	RDG_BUFFER_ACCESS_ARRAY(InputBuffers)
+	RDG_BUFFER_ACCESS_ARRAY(OutputBuffers)
 END_SHADER_PARAMETER_STRUCT()
 
 //
@@ -1177,27 +1178,11 @@ END_SHADER_PARAMETER_STRUCT()
 //
 void FModel::AddDispatchOps_RenderThread(FRDGBuilder& GraphBuilder)
 {
-	const ERDGPassFlags TransitionBuffFlags = ERDGPassFlags::Compute | ERDGPassFlags::NeverCull;
-
-	InputBuffers.Reset();
-	OutputBuffers.Reset();
-	InputBuffers.Reserve(InputTensorIndices.Num() + WeightTensorIndices.Num());
-	OutputBuffers.Reserve(OutputTensorIndices.Num());
+	FDmlModelDispatchPassParameters* DispatchParams = GraphBuilder.AllocParameters<FDmlModelDispatchPassParameters>();
 
 	for (int32 Idx = 0; Idx < InputTensorIndices.Num(); ++Idx)
 	{
-		FTensorBufferParamsDml* Params = GraphBuilder.AllocParameters<FTensorBufferParamsDml>();
-		Params->Buffer = AllTensorRDGRefs[InputTensorIndices[Idx]]->GetBuffer();
-
-		GraphBuilder.AddPass(
-			RDG_EVENT_NAME("FModel_Dispatch_GetInputBuffer"),
-			Params,
-			TransitionBuffFlags,
-			[this, Idx, Params](FRHICommandListImmediate& RHICmdList)
-			{
-				InputBuffers.Insert(Params->Buffer->GetRHI(), Idx);
-			}
-		);
+		DispatchParams->InputBuffers.Emplace(AllTensorRDGRefs[InputTensorIndices[Idx]]->GetBuffer(), ERHIAccess::UAVCompute);
 	}
 
 	int32 NumWeightTensors = 0;
@@ -1212,36 +1197,41 @@ void FModel::AddDispatchOps_RenderThread(FRDGBuilder& GraphBuilder)
 
 	for (int32 Idx = 0; Idx < OutputTensorIndices.Num(); ++Idx)
 	{
-		FTensorBufferParamsDml* Params = GraphBuilder.AllocParameters<FTensorBufferParamsDml>();
-		Params->Buffer = AllTensorRDGRefs[OutputTensorIndices[Idx]]->GetBuffer();
-
-		GraphBuilder.AddPass(
-			RDG_EVENT_NAME("FModel_Dispatch_GetOutputBuffer"),
-			Params,
-			TransitionBuffFlags,
-			[this, Idx, Params](FRHICommandListImmediate& RHICmdList)
-			{
-				OutputBuffers.Insert(Params->Buffer->GetRHI(), Idx);
-			}
-		);
+		DispatchParams->OutputBuffers.Emplace(AllTensorRDGRefs[OutputTensorIndices[Idx]]->GetBuffer(), ERHIAccess::UAVCompute);
 	}
-	
+
 	GraphBuilder.AddPass(
-		RDG_EVENT_NAME("FDmlModel_Dispatch"),
-		ERDGPassFlags::None | ERDGPassFlags::NeverCull,
-		[this, NumWeightTensors](FRHICommandListImmediate& RHICmdList)
+		RDG_EVENT_NAME("FDmlModelDispatch"),
+		DispatchParams,
+		ERDGPassFlags::Compute | ERDGPassFlags::NeverCull,
+		[this, DispatchParams, NumWeightTensors](FRHICommandListImmediate& RHICmdList)
 		{
+			FRHIBufferInputArray	RHIInputBuffers;
+			FRHIBufferOutputArray	RHIOutputBuffers;
+
+			for (FRDGBuffer* RDGBuffer : DispatchParams->InputBuffers)
+			{
+				RDGBuffer->MarkResourceAsUsed();
+				RHIInputBuffers.Add(RDGBuffer->GetRHI());
+			}
+
+			for (int32 Idx = 0; Idx < NumWeightTensors; ++Idx)
+			{
+				RHIInputBuffers.Add(nullptr);
+			}
+
+			for (FRDGBuffer* RDGBuffer : DispatchParams->OutputBuffers)
+			{
+				RDGBuffer->MarkResourceAsUsed();
+				RHIOutputBuffers.Add(RDGBuffer->GetRHI());
+			}
+
 			RHICmdList.EnqueueLambda(
-				[this, NumWeightTensors](FRHICommandListImmediate& RHICmdList)
+				[this, InputBuffers = MoveTemp(RHIInputBuffers), OutputBuffers = MoveTemp(RHIOutputBuffers)](FRHICommandListImmediate& RHICmdList)
 				{
 					TArray<CD3DX12_RESOURCE_BARRIER, TInlineAllocator<MaxNumInputs + MaxNumOutputs>>	PreBarriers;
 					TArray<CD3DX12_RESOURCE_BARRIER, TInlineAllocator<MaxNumOutputs * 2>>				PostBarriers;
 
-					for (int32 Idx = 0; Idx < NumWeightTensors; ++Idx)
-					{
-						InputBuffers.Add(nullptr);
-					}
-
 					for (FRHIBuffer* Buffer : InputBuffers)
 					{
 						if (!Buffer)
diff --git a/Engine/Plugins/Experimental/NNE/Source/NNERuntimeRDG/Private/Dml/NNEDmlModel.h b/Engine/Plugins/Experimental/NNE/Source/NNERuntimeRDG/Private/Dml/NNEDmlModel.h
--- a/Engine/Plugins/Experimental/NNE/Source/NNERuntimeRDG/Private/Dml/NNEDmlModel.h
+++ b/Engine/Plugins/Experimental/NNE/Source/NNERuntimeRDG/Private/Dml/NNEDmlModel.h
@@ -70,7 +70,7 @@ private:
 	ID3D12Resource* CreateD3D12Buffer(uint32 Size, D3D12_RESOURCE_STATES ResourceState = D3D12_RESOURCE_STATE_COMMON, D3D12_HEAP_TYPE HeapType = D3D12_HEAP_TYPE_DEFAULT, const TCHAR* DebugName = nullptr);
 
 	//Note: This should go into RDG
-	static constexpr int32 MaxNumInputs = 32;
+	static constexpr int32 MaxNumInputs = 512;
 	static constexpr int32 MaxNumOutputs = 4;
 
 	using FRHIBufferInputArray = TArray<FRHIBuffer*, TInlineAllocator<MaxNumInputs>>;
@@ -86,8 +86,6 @@ private:
 
 	TArray<int32>						ConstantCPUTensorIndices;
 
-	FRHIBufferInputArray				InputBuffers;
-	FRHIBufferOutputArray				OutputBuffers;
 #ifdef NNE_USE_D3D12_RESOURCES
 	TComPtr<ID3D12Resource>				PersistBuff;
 	TComPtr<ID3D12Resource>				TempBuff;
