Add RDG Dml Squeeze operator, missed by previous change

#rnx

diff --git a/Engine/Plugins/Experimental/NNE/Source/NNERuntimeRDG/Private/Dml/NNEDmlOperatorSqueeze.cpp b/Engine/Plugins/Experimental/NNE/Source/NNERuntimeRDG/Private/Dml/NNEDmlOperatorSqueeze.cpp
new file mode 100644
--- /dev/null
+++ b/Engine/Plugins/Experimental/NNE/Source/NNERuntimeRDG/Private/Dml/NNEDmlOperatorSqueeze.cpp
@@ -0,0 +1,93 @@
+// Copyright Epic Games, Inc. All Rights Reserved.
+
+#ifdef NNE_USE_DIRECTML
+
+#include "NNEDmlOperator.h"
+#include "NNEDmlOperatorUtils.h"
+
+namespace UE::NNERuntimeRDG::Private::Dml
+{
+
+class FOperatorDmlSqueeze : public FOperatorDml
+{
+
+public:
+
+	static FOperatorDml* Create()
+	{
+		return new FOperatorDmlSqueeze();
+	}
+
+	//
+	//
+	//
+	virtual bool Initialize(IDMLDevice* Device, TArrayView<const NNECore::Internal::FTensor> InputTensors, TArrayView<const NNECore::Internal::FTensor> OutputTensors, const NNECore::FAttributeMap& Attributes) override
+	{
+		check(InputTensors.Num() == 1);
+		check(OutputTensors.Num() == 1);
+
+		TConstArrayView<uint32>		InputShape = InputTensors[0].GetShape().GetData();
+		TArray<int32>				Axes;
+		const FNNEAttributeValue*	AxesAttr = Attributes.GetAttributeValue(TEXT("axes"));
+		
+		if (AxesAttr)
+		{
+			Axes = AxesAttr->GetValue<TArray<int32>>();
+			HandleNegativeAxes(Axes, InputShape.Num());
+			Algo::Sort(Axes);
+		}
+		
+		if (Axes.IsEmpty())
+		{
+			for (int32 Idx = 0; Idx < InputShape.Num(); ++Idx)
+			{
+				if (InputShape[Idx] == 1)
+				{
+					Axes.Add(Idx);
+				}
+			}
+		}
+		
+		DmlUtil::FSmallUIntArray	OutputShape;
+
+		OutputShape.Reserve(InputShape.Num());
+		OutputShape.Append(InputShape);
+		
+		for (int32 Idx = Axes.Num() - 1; Idx >= 0; --Idx)
+		{
+			const int32 Axe = Axes[Idx];
+
+			if (OutputShape[Axe] != 1)
+			{
+				UE_LOG(LogNNE, Warning, TEXT("Squeeze at axe %d for 'Data' (name: %s) should be targeting a dimension of size 1 but it is %d."), Axe, *InputTensors[0].GetName(), OutputShape[Axe]);
+				return false;
+			}
+
+			OutputShape.RemoveAt(Axe);
+		}
+
+		DmlUtil::FTensorDesc DmlTensorDesc;
+		
+		if (!DmlTensorDesc.InitFromTensor(OutputTensors[0], OutputShape.Num(),
+			/*Broadcast =*/ MakeEmptyArrayView<uint32>(),
+			/*CustomShape =*/ OutputShape))
+		{
+			UE_LOG(LogNNE, Warning, TEXT("Failed to initialize Unsqueeze's output tensor for DML inference"));
+			return false;
+		}
+
+		DML_ELEMENT_WISE_IDENTITY_OPERATOR_DESC DmlIdentityOpDesc{};
+
+		DmlIdentityOpDesc.InputTensor = &DmlTensorDesc.Desc;
+		DmlIdentityOpDesc.OutputTensor = &DmlTensorDesc.Desc;
+
+		return CreateOperator(Device, DML_OPERATOR_DESC{ DML_OPERATOR_ELEMENT_WISE_IDENTITY, &DmlIdentityOpDesc });
+	}
+};
+
+// Register Reshape operator on Module startup
+NNE_DML_REGISTER_OP(Squeeze)
+
+} // namespace UE::NNERuntimeRDG::Private::Dml
+
+#endif // NNE_USE_DIRECTML
