Fix clang compiler errors

#rnx

diff --git a/Engine/Plugins/Experimental/NNE/Source/NNERuntimeRDG/Private/Dml/NNEDmlOperatorReduce.cpp b/Engine/Plugins/Experimental/NNE/Source/NNERuntimeRDG/Private/Dml/NNEDmlOperatorReduce.cpp
--- a/Engine/Plugins/Experimental/NNE/Source/NNERuntimeRDG/Private/Dml/NNEDmlOperatorReduce.cpp
+++ b/Engine/Plugins/Experimental/NNE/Source/NNERuntimeRDG/Private/Dml/NNEDmlOperatorReduce.cpp
@@ -132,7 +132,7 @@ public:
 };
 
 // Register Reshape operator on Module startup
-#define OP(OpName, ReduceFunc) FOperatorRegistryDml::Get()->OpAdd(TEXT(#OpName), FOperatorDmlReduce<##ReduceFunc##>::Create)
+#define OP(OpName, ReduceFunc) FOperatorRegistryDml::Get()->OpAdd(TEXT(#OpName), FOperatorDmlReduce<ReduceFunc>::Create)
 
 struct FOperatorDmlReduceRegistrator
 {
