NNE - Add HS models to tests

#rnx

diff --git a/Engine/Plugins/Experimental/NNE/Source/NNEQA/Resources/NNEQATestDesc.json b/Engine/Plugins/Experimental/NNE/Source/NNEQA/Resources/NNEQATestDesc.json
--- a/Engine/Plugins/Experimental/NNE/Source/NNEQA/Resources/NNEQATestDesc.json
+++ b/Engine/Plugins/Experimental/NNE/Source/NNEQA/Resources/NNEQATestDesc.json
@@ -61,7 +61,85 @@
 					"absolute_tolerance": 1e-4
 				},
 				{
-					"target": "FaceDetector",
+					"target": "hs_m01",
+					"runtimes": [
+						{ "name": "NNERuntimeRDGHlsl", "skip": true, "comment": "Model not supported yet." }
+					],
+					"absolute_tolerance": 1e-4,
+					"skip": true, "comment": "Not working yet"
+				},
+				{
+					"target": "hs_m02",
+					"runtimes": [
+						{ "name": "NNERuntimeRDGHlsl", "skip": true, "comment": "Model not supported yet." }
+					],
+					"absolute_tolerance": 1e-4
+				},
+				{
+					"target": "hs_m03",
+					"runtimes": [
+						{ "name": "NNERuntimeRDGHlsl", "skip": true, "comment": "Model not supported yet." }
+					],
+					"absolute_tolerance": 1e-4
+				},
+				{
+					"target": "hs_m04",
+					"runtimes": [
+						{ "name": "NNERuntimeRDGHlsl", "skip": true, "comment": "Model not supported yet." }
+					],
+					"absolute_tolerance": 1e-4,
+					"skip": true, "comment": "Not working yet"
+				},
+				{
+					"target": "hs_m05",
+					"runtimes": [
+						{ "name": "NNERuntimeRDGHlsl", "skip": true, "comment": "Model not supported yet." }
+					],
+					"absolute_tolerance": 1e-4,
+					"skip": true, "comment": "Not working yet"
+				},
+				{
+					"target": "hs_m06",
+					"runtimes": [
+						{ "name": "NNERuntimeRDGHlsl", "skip": true, "comment": "Model not supported yet." }
+					],
+					"absolute_tolerance": 1e-4,
+					"skip": true, "comment": "Not working yet"
+				},
+				{
+					"target": "hs_m07",
+					"runtimes": [
+						{ "name": "NNERuntimeRDGHlsl", "skip": true, "comment": "Model not supported yet." }
+					],
+					"absolute_tolerance": 1e-4,
+					"skip": true, "comment": "Not working yet"
+				},
+				{
+					"target": "hs_m08",
+					"runtimes": [
+						{ "name": "NNERuntimeRDGHlsl", "skip": true, "comment": "Model not supported yet." }
+					],
+					"absolute_tolerance": 1e-4,
+					"skip": true, "comment": "Not working yet"
+				},
+				{
+					"target": "hs_m09",
+					"runtimes": [
+						{ "name": "NNERuntimeRDGHlsl", "skip": true, "comment": "Model not supported yet." }
+					],
+					"absolute_tolerance": 1e-4,
+					"skip": true, "comment": "Not working yet"
+				},
+				{
+					"target": "hs_m10",
+					"runtimes": [
+						{ "name": "NNERuntimeRDGHlsl", "skip": true, "comment": "Model not supported yet." }
+					],
+					"absolute_tolerance": 1e-4,
+					"skip": true, "comment": "Not working yet"
+				},
+				{
+					"target": "hs_m11",
 					"runtimes": [
 						{ "name": "NNERuntimeRDGHlsl", "skip": true, "comment": "Model not supported yet." }
 					],
