NNE - Add RDG Dml GlobalAveragePool

#rnx

diff --git a/Engine/Plugins/Experimental/NNE/Source/NNEQA/Resources/NNEQATestDesc.json b/Engine/Plugins/Experimental/NNE/Source/NNEQA/Resources/NNEQATestDesc.json
--- a/Engine/Plugins/Experimental/NNE/Source/NNEQA/Resources/NNEQATestDesc.json
+++ b/Engine/Plugins/Experimental/NNE/Source/NNEQA/Resources/NNEQATestDesc.json
@@ -540,6 +540,15 @@
 			"runtimes": [
 				{ "name": "NNERuntimeRDGHlsl", "skip": true, "comment": "Not supported" }
 			]
+		},
+		{
+			"category": "GlobalAveragePoolOp",
+			"targets": [
+				{ "target": "GlobalAveragePool", "absolute_tolerance": 1e-6 }
+			],
+			"runtimes": [
+				{ "name": "NNERuntimeRDGHlsl", "skip": true, "comment": "Not supported" }
+			]
 		}
 	],
 
@@ -1189,6 +1198,18 @@
 				{ "inputs": [ { "shape": [ 1, 2, 256, 256 ] } ], "outputs": [ { "shape": [ 1, 2, 128, 128 ] } ] }
 			]
 		},
+		{
+			"name": "GlobalAveragePoolOp.2D",
+			"datasets": [
+				{ "inputs": [ { "shape": [ 2, 2, 128, 256 ] } ], "outputs": [ { "shape": [ 2, 2, 1, 1 ] } ] }
+			]
+		},
+		{
+			"name": "GlobalAveragePoolOp.3D",
+			"datasets": [
+				{ "inputs": [ { "shape": [ 2, 3, 128, 256, 64 ] } ], "outputs": [ { "shape": [ 2, 3, 1, 1, 1 ] } ] }
+			]
+		},
 		{ 
 			"name": "Reduce.Default",
 			"datasets": [
diff --git a/Engine/Plugins/Experimental/NNE/Source/NNERuntimeRDG/Private/Dml/NNEDmlOperatorPool.cpp b/Engine/Plugins/Experimental/NNE/Source/NNERuntimeRDG/Private/Dml/NNEDmlOperatorPool.cpp
--- a/Engine/Plugins/Experimental/NNE/Source/NNERuntimeRDG/Private/Dml/NNEDmlOperatorPool.cpp
+++ b/Engine/Plugins/Experimental/NNE/Source/NNERuntimeRDG/Private/Dml/NNEDmlOperatorPool.cpp
@@ -13,7 +13,7 @@ namespace UE::NNERuntimeRDG::Private::Dml
 /**
 	* MaxPool
 	*/
-template <typename TDmlPoolOpDesc, DML_OPERATOR_TYPE DmlOpType>
+template <typename TDmlPoolOpDesc, DML_OPERATOR_TYPE DmlOpType, bool UseGlobalPooling>
 class FOperatorDmlPool : public FOperatorDml
 {
 	using FIntArray = TArray<int32>;
@@ -22,7 +22,7 @@ public:
 
 	static FOperatorDml* Create()
 	{
-		return new FOperatorDmlPool<TDmlPoolOpDesc, DmlOpType>();
+		return new FOperatorDmlPool<TDmlPoolOpDesc, DmlOpType, UseGlobalPooling>();
 	}
 
 	//
@@ -46,26 +46,40 @@ public:
 		DmlUtil::FSmallUIntArray StartPadding, EndPadding, KernelShape, Strides, Dilations;
 
 		Strides.Init(1u, NumSpatialDimensions);
-		if (!DmlUtil::GetArrayAttributeNoOverflow(Attributes.GetAttributeValue(TEXT("strides")), Strides, TConstArrayView<uint32>(Strides)))
+		if constexpr(!UseGlobalPooling)
 		{
-			UE_LOG(LogNNE, Error, TEXT("Strides attribute cast led to overflow"));
-			return false;
+			if (!DmlUtil::GetArrayAttributeNoOverflow(Attributes.GetAttributeValue(TEXT("strides")), Strides, TConstArrayView<uint32>(Strides)))
+			{
+				UE_LOG(LogNNE, Error, TEXT("Strides attribute cast led to overflow"));
+				return false;
+			}
 		}
 		check(Strides.Num() == NumSpatialDimensions);
 
 		Dilations.Init(1u, NumSpatialDimensions);
-		if (!DmlUtil::GetArrayAttributeNoOverflow(Attributes.GetAttributeValue(TEXT("dilations")), Dilations, TConstArrayView<uint32>(Dilations)))
+		if constexpr(!UseGlobalPooling)
 		{
-			UE_LOG(LogNNE, Error, TEXT("Dilations attribute cast led to overflow"));
-			return false;
+			if (!DmlUtil::GetArrayAttributeNoOverflow(Attributes.GetAttributeValue(TEXT("dilations")), Dilations, TConstArrayView<uint32>(Dilations)))
+			{
+				UE_LOG(LogNNE, Error, TEXT("Dilations attribute cast led to overflow"));
+				return false;
+			}
 		}
 		check(Dilations.Num() == NumSpatialDimensions);
-
-		if (!DmlUtil::GetArrayAttributeNoOverflow(Attributes.GetAttributeValue(TEXT("kernel_shape")), KernelShape))
+		
+		if constexpr(UseGlobalPooling)
 		{
-			UE_LOG(LogNNE, Error, TEXT("kernel_shape attribute cast led to overflow"));
-			return false;
+			KernelShape = DmlUtil::FSmallUIntArray{ InputShape.RightChop(NonspatialDimensionCount) };
+		}
+		else
+		{
+			if (!DmlUtil::GetArrayAttributeNoOverflow(Attributes.GetAttributeValue(TEXT("kernel_shape")), KernelShape))
+			{
+				UE_LOG(LogNNE, Error, TEXT("kernel_shape attribute cast led to overflow"));
+				return false;
+			}
 		}
+		
 		if (KernelShape.IsEmpty())
 		{
 			UE_LOG(LogNNE, Error, TEXT("kernel_shape attribute is required for pooling operators"));
@@ -73,14 +87,22 @@ public:
 		}
 		check(KernelShape.Num() == NumSpatialDimensions);
 
-		ComputeStartEndPaddings(
-			InputShape,
-			Attributes,
-			StartPadding,
-			EndPadding,
-			KernelPadding(InputShape, KernelShape, Dilations, Strides)
-		);
-
+		if constexpr(UseGlobalPooling)
+		{
+			StartPadding.Init(0u, NumSpatialDimensions);
+			EndPadding.Init(0u, NumSpatialDimensions);
+		}
+		else
+		{
+			ComputeStartEndPaddings(
+				InputShape,
+				Attributes,
+				StartPadding,
+				EndPadding,
+				KernelPadding(InputShape, KernelShape, Dilations, Strides)
+			);
+		}
+		
 		for (int Idx = 0; Idx < OutputTensors.Num(); ++Idx)
 		{
 			check(OutputTensors[Idx].GetShape().Rank() == InputShape.Num());
@@ -171,20 +193,21 @@ public:
 	}
 };
 
-#define NNE_DML_REGISTER_POOLING_OP(OpName, DmlPrefix) \
+#define NNE_DML_REGISTER_POOLING_OP(OpName, DmlPrefix, UseGlobalPooling) \
 struct FDmlOperator##OpName##Registrator \
 { \
 	FDmlOperator##OpName##Registrator() \
 	{ \
-		FOperatorRegistryDml::Get()->OpAdd(TEXT(#OpName), FOperatorDmlPool<DML_##DmlPrefix##_POOLING_OPERATOR_DESC, DML_OPERATOR_##DmlPrefix##_POOLING>::Create); \
+		FOperatorRegistryDml::Get()->OpAdd(TEXT(#OpName), FOperatorDmlPool<DML_##DmlPrefix##_POOLING_OPERATOR_DESC, DML_OPERATOR_##DmlPrefix##_POOLING, UseGlobalPooling>::Create); \
 	} \
 }; \
 \
 static FDmlOperator##OpName##Registrator RegisterDmlOperator##OpName;
 
 // Register pooling operator on Module startup
-NNE_DML_REGISTER_POOLING_OP(MaxPool, MAX)
-NNE_DML_REGISTER_POOLING_OP(AveragePool, AVERAGE)
+NNE_DML_REGISTER_POOLING_OP(MaxPool, MAX, false)
+NNE_DML_REGISTER_POOLING_OP(AveragePool, AVERAGE, false)
+NNE_DML_REGISTER_POOLING_OP(GlobalAveragePool, AVERAGE, true)
 
 #undef NNE_DML_REGISTER_POOLING_OP
 
