NNE fixes
* Support for const output tensors (at the format level, works for Dml and Hlsl)
* Dml : Added Identity op + Support for scalar tensor when initializing tensor descs

Tests:
* Enabled all hs models
* Added a dedicated test for const output for non regression
* All nne test green

#rnx

diff --git a/Engine/Plugins/Experimental/NNE/Source/NNEQA/Resources/NNEQATestDesc.json b/Engine/Plugins/Experimental/NNE/Source/NNEQA/Resources/NNEQATestDesc.json
--- a/Engine/Plugins/Experimental/NNE/Source/NNEQA/Resources/NNEQATestDesc.json
+++ b/Engine/Plugins/Experimental/NNE/Source/NNEQA/Resources/NNEQATestDesc.json
@@ -65,8 +65,7 @@
 					"runtimes": [
 						{ "name": "NNERuntimeRDGHlsl", "skip": true, "comment": "Model not supported yet." }
 					],
-					"absolute_tolerance": 1e-4,
-					"skip": true, "comment": "Not working yet"
+					"absolute_tolerance": 1e-4
 				},
 				{
 					"target": "hs_m02",
@@ -87,62 +86,48 @@
 					"runtimes": [
 						{ "name": "NNERuntimeRDGHlsl", "skip": true, "comment": "Model not supported yet." }
 					],
-					"absolute_tolerance": 1e-4,
-					"skip": true, "comment": "Not working yet"
+					"absolute_tolerance": 1e-4
 				},
 				{
 					"target": "hs_m05",
 					"runtimes": [
 						{ "name": "NNERuntimeRDGHlsl", "skip": true, "comment": "Model not supported yet." }
 					],
-					"absolute_tolerance": 1e-4,
-					"skip": true, "comment": "Not working yet"
+					"absolute_tolerance": 1e-4
 				},
 				{
 					"target": "hs_m06",
 					"runtimes": [
 						{ "name": "NNERuntimeRDGHlsl", "skip": true, "comment": "Model not supported yet." }
 					],
-					"absolute_tolerance": 1e-4,
-					"skip": true, "comment": "Not working yet"
+					"absolute_tolerance": 1e-4
 				},
 				{
 					"target": "hs_m07",
 					"runtimes": [
 						{ "name": "NNERuntimeRDGHlsl", "skip": true, "comment": "Model not supported yet." }
 					],
-					"absolute_tolerance": 1e-4,
-					"skip": true, "comment": "Not working yet"
+					"absolute_tolerance": 1e-4
 				},
 				{
 					"target": "hs_m08",
 					"runtimes": [
 						{ "name": "NNERuntimeRDGHlsl", "skip": true, "comment": "Model not supported yet." }
 					],
-					"absolute_tolerance": 1e-4,
-					"skip": true, "comment": "Not working yet"
+					"absolute_tolerance": 1e-4
 				},
 				{
 					"target": "hs_m09",
 					"runtimes": [
 						{ "name": "NNERuntimeRDGHlsl", "skip": true, "comment": "Model not supported yet." }
 					],
-					"absolute_tolerance": 1e-4,
-					"skip": true, "comment": "Not working yet"
+					"absolute_tolerance": 1e-4
 				},
 				{
 					"target": "hs_m10",
 					"runtimes": [
 						{ "name": "NNERuntimeRDGHlsl", "skip": true, "comment": "Model not supported yet." }
 					],
-					"absolute_tolerance": 1e-4,
-					"skip": true, "comment": "Not working yet"
-				},
-				{
-					"target": "hs_m11",
-					"runtimes": [
-						{ "name": "NNERuntimeRDGHlsl", "skip": true, "comment": "Model not supported yet." }
-					],
 					"absolute_tolerance": 1e-4
 				}
 			]
@@ -168,6 +153,10 @@
 					"runtimes": [
 						{ "name": "NNERuntimeRDGDml",  "skip": true, "comment": "Model not supported yet." }
 					]
+				},
+				{
+					"target": "ConstOutput",
+					"comment": "Model contain a 2nd output (Y1) constant from an initializer."
 				}
 			]
 		},
@@ -396,9 +385,6 @@
 		},
 		{
 			"category": "Identity",
-			"runtimes": [
-				{ "name": "NNERuntimeRDGDml", "skip": true },
-			],
 			"targets": [
 				{ "target": "Identity" }
 			]
diff --git a/Engine/Plugins/Experimental/NNE/Source/NNERuntimeRDG/Private/Dml/NNEDmlOperator.cpp b/Engine/Plugins/Experimental/NNE/Source/NNERuntimeRDG/Private/Dml/NNEDmlOperator.cpp
--- a/Engine/Plugins/Experimental/NNE/Source/NNERuntimeRDG/Private/Dml/NNEDmlOperator.cpp
+++ b/Engine/Plugins/Experimental/NNE/Source/NNERuntimeRDG/Private/Dml/NNEDmlOperator.cpp
@@ -389,6 +389,16 @@ bool FOperatorDml::InitDmlTensorDesc(DmlUtil::FTensorDesc& DmlTensorDesc, const
 	BuffDesc.Strides = nullptr;
 	BuffDesc.TotalTensorSizeInBytes = Tensor.GetDataSize(); // DmlUtil::CalculateBufferSize(DmlTensorDesc, Tensor);
 
+	static DmlUtil::FSmallUIntArray ScalarShape({ 1 });
+
+	//Handle scalar tensors
+	if (Tensor.GetShape().Rank() == 0)
+	{
+		BuffDesc.DimensionCount = ScalarShape.Num();
+		DmlTensorDesc.Sizes = ScalarShape;
+		BuffDesc.Sizes = ScalarShape.GetData();
+	}
+
 	DmlTensorDesc.Desc = DML_TENSOR_DESC{ DML_TENSOR_TYPE_BUFFER, &DmlTensorDesc.BuffDesc };
 
 	return true;
diff --git a/Engine/Plugins/Experimental/NNE/Source/NNERuntimeRDG/Private/Dml/NNEDmlOperatorIdentity.cpp b/Engine/Plugins/Experimental/NNE/Source/NNERuntimeRDG/Private/Dml/NNEDmlOperatorIdentity.cpp
new file mode 100644
--- /dev/null
+++ b/Engine/Plugins/Experimental/NNE/Source/NNERuntimeRDG/Private/Dml/NNEDmlOperatorIdentity.cpp
@@ -0,0 +1,60 @@
+// Copyright Epic Games, Inc. All Rights Reserved.
+
+#ifdef NNE_USE_DIRECTML
+
+#include "NNEDmlOperator.h"
+#include "NNEDmlOperatorUtils.h"
+
+namespace UE::NNERuntimeRDG::Private::Dml
+{
+
+class FOperatorDmlIdentity : public FOperatorDml
+{
+
+public:
+
+	static FOperatorDml* Create()
+	{
+		return new FOperatorDmlIdentity();
+	}
+
+	//
+	//
+	//
+	virtual bool Initialize(IDMLDevice* Device, TArrayView<const NNECore::Internal::FTensor> InputTensors, TArrayView<const NNECore::Internal::FTensor> OutputTensors, const NNECore::FAttributeMap& Attributes) override
+	{
+		check(InputTensors.Num() == 1);
+		check(OutputTensors.Num() == 1);
+
+		const NNECore::Internal::FTensor& InputTensorDesc = InputTensors[0];
+		const NNECore::Internal::FTensor& OutputTensorDesc = OutputTensors[0];
+
+		DmlUtil::FTensorDesc	DmlInputTensorDesc{};
+		DmlUtil::FTensorDesc	DmlOutputTensorDesc{};
+
+		if (!InitDmlTensorDesc(DmlInputTensorDesc, InputTensorDesc))
+		{
+			UE_LOG(LogNNE, Warning, TEXT("Failed to initialize input tensor for DML inference"));
+			return false;
+		}
+		if (!InitDmlTensorDesc(DmlOutputTensorDesc, OutputTensorDesc))
+		{
+			UE_LOG(LogNNE, Warning, TEXT("Failed to initialize output tensor for DML inference"));
+			return false;
+		}
+
+		DML_ELEMENT_WISE_IDENTITY_OPERATOR_DESC DmlIdentityOpDesc{};
+
+		DmlIdentityOpDesc.InputTensor = &DmlInputTensorDesc.Desc;
+		DmlIdentityOpDesc.OutputTensor = &DmlOutputTensorDesc.Desc;
+
+		return CreateOperator(Device, DML_OPERATOR_DESC{ DML_OPERATOR_ELEMENT_WISE_IDENTITY, &DmlIdentityOpDesc });
+	}
+};
+
+// Register Reshape operator on Module startup
+NNE_DML_REGISTER_OP(Identity)
+
+} // namespace UE::NNERuntimeRDG::Private::Dml
+
+#endif // NNE_USE_DIRECTML
diff --git a/Engine/Plugins/Experimental/NNE/Source/NNEUtils/Private/NNEUtilsModelOptimizer.cpp b/Engine/Plugins/Experimental/NNE/Source/NNEUtils/Private/NNEUtilsModelOptimizer.cpp
--- a/Engine/Plugins/Experimental/NNE/Source/NNEUtils/Private/NNEUtilsModelOptimizer.cpp
+++ b/Engine/Plugins/Experimental/NNE/Source/NNEUtils/Private/NNEUtilsModelOptimizer.cpp
@@ -525,6 +525,8 @@ private:
 			Builder->AddInput(Tensor);
 		}
 
+		TArray<uint8> TensorDataBlob;
+
 		// Add tensors for graph outputs
 		for (int Idx = 0; Idx < GraphInfo.outputCount; ++Idx)
 		{
@@ -540,11 +542,37 @@ private:
 				);
 
 			Builder->AddOutput(Tensor);
+
+			Ort::GraphTensorInitializer	TensorInit = Graph->GetTensorInitializer(TensorInfo.name);
+			if (TensorInit)
+			{
+				const void* Data = nullptr;
+				uint64	DataSize = 0;
+
+				DataSize = Graph->GetTensorDataSize(TensorInit);
+				if (TensorDataBlob.Num() < DataSize)
+				{
+					TensorDataBlob.SetNumUninitialized(DataSize);
+				}
+
+				Graph->GetTensorData(TensorInit, TensorDataBlob.GetData(), DataSize, 0);
+				Data = TensorDataBlob.GetData();
+
+				auto TensorInitializer =
+					Builder->AddTensor(
+						FString(ANSI_TO_TCHAR(TensorInfo.name)) + TEXT("_NNEInitializer"),
+						DataType,
+						MakeArrayView(TensorInfo.shape, TensorInfo.shapeLen),
+						Data, DataSize
+					);
+
+				auto Op = Builder->AddOperator(TEXT("Identity"));
+				Builder->AddOperatorInput(Op, TensorInitializer);
+				Builder->AddOperatorOutput(Op, Tensor);
+			}
 		}
 
 		// Traverse all the nodes get their inputs, outputs and tensor data
-		TArray<uint8> TensorDataBlob;
-
 		for (int Idx = 0; Idx < GraphInfo.nodeCount; ++Idx)
 		{
 			Ort::GraphNode		Node = Graph->GetNode(Idx);
