NNE - Fix GlobalAveragePool tests

#rnx

diff --git a/Engine/Plugins/Experimental/NNE/Source/NNEQA/Resources/NNEQATestDesc.json b/Engine/Plugins/Experimental/NNE/Source/NNEQA/Resources/NNEQATestDesc.json
--- a/Engine/Plugins/Experimental/NNE/Source/NNEQA/Resources/NNEQATestDesc.json
+++ b/Engine/Plugins/Experimental/NNE/Source/NNEQA/Resources/NNEQATestDesc.json
@@ -528,7 +528,7 @@
 			]
 		},
 		{
-			"category": "GlobalAveragePoolOp",
+			"category": "GlobalAveragePool",
 			"targets": [
 				{ "target": "GlobalAveragePool", "absolute_tolerance": 1e-6 }
 			],
@@ -1185,13 +1185,13 @@
 			]
 		},
 		{
-			"name": "GlobalAveragePoolOp.2D",
+			"name": "GlobalAveragePool.2D",
 			"datasets": [
 				{ "inputs": [ { "shape": [ 2, 2, 128, 256 ] } ], "outputs": [ { "shape": [ 2, 2, 1, 1 ] } ] }
 			]
 		},
 		{
-			"name": "GlobalAveragePoolOp.3D",
+			"name": "GlobalAveragePool.3D",
 			"datasets": [
 				{ "inputs": [ { "shape": [ 2, 3, 128, 256, 64 ] } ], "outputs": [ { "shape": [ 2, 3, 1, 1, 1 ] } ] }
 			]
