NNE - Add Dml GlobalMaxPool

#rnx

diff --git a/Engine/Plugins/Experimental/NNE/Source/NNEQA/Resources/NNEQATestDesc.json b/Engine/Plugins/Experimental/NNE/Source/NNEQA/Resources/NNEQATestDesc.json
--- a/Engine/Plugins/Experimental/NNE/Source/NNEQA/Resources/NNEQATestDesc.json
+++ b/Engine/Plugins/Experimental/NNE/Source/NNEQA/Resources/NNEQATestDesc.json
@@ -528,8 +528,9 @@
 			]
 		},
 		{
-			"category": "GlobalAveragePool",
+			"category": "GlobalPool",
 			"targets": [
+				{ "target": "GlobalMaxPool" },
 				{ "target": "GlobalAveragePool", "absolute_tolerance": 1e-6 }
 			],
 			"runtimes": [
@@ -1185,13 +1186,13 @@
 			]
 		},
 		{
-			"name": "GlobalAveragePool.2D",
+			"name": "GlobalPool.2D",
 			"datasets": [
 				{ "inputs": [ { "shape": [ 2, 2, 128, 256 ] } ], "outputs": [ { "shape": [ 2, 2, 1, 1 ] } ] }
 			]
 		},
 		{
-			"name": "GlobalAveragePool.3D",
+			"name": "GlobalPool.3D",
 			"datasets": [
 				{ "inputs": [ { "shape": [ 2, 3, 128, 256, 64 ] } ], "outputs": [ { "shape": [ 2, 3, 1, 1, 1 ] } ] }
 			]
diff --git a/Engine/Plugins/Experimental/NNE/Source/NNERuntimeRDG/Private/Dml/NNEDmlOperatorPool.cpp b/Engine/Plugins/Experimental/NNE/Source/NNERuntimeRDG/Private/Dml/NNEDmlOperatorPool.cpp
--- a/Engine/Plugins/Experimental/NNE/Source/NNERuntimeRDG/Private/Dml/NNEDmlOperatorPool.cpp
+++ b/Engine/Plugins/Experimental/NNE/Source/NNERuntimeRDG/Private/Dml/NNEDmlOperatorPool.cpp
@@ -206,6 +206,7 @@ static FDmlOperator##OpName##Registrator RegisterDmlOperator##OpName;
 
 // Register pooling operator on Module startup
 NNE_DML_REGISTER_POOLING_OP(MaxPool, MAX, false)
+NNE_DML_REGISTER_POOLING_OP(GlobalMaxPool, MAX, true)
 NNE_DML_REGISTER_POOLING_OP(AveragePool, AVERAGE, false)
 NNE_DML_REGISTER_POOLING_OP(GlobalAveragePool, AVERAGE, true)
 
