NNE - Add Dml LpPooling

#rnx

diff --git a/Engine/Plugins/Experimental/NNE/Source/NNEQA/Resources/NNEQATestDesc.json b/Engine/Plugins/Experimental/NNE/Source/NNEQA/Resources/NNEQATestDesc.json
--- a/Engine/Plugins/Experimental/NNE/Source/NNEQA/Resources/NNEQATestDesc.json
+++ b/Engine/Plugins/Experimental/NNE/Source/NNEQA/Resources/NNEQATestDesc.json
@@ -521,7 +521,8 @@
 			"category": "PoolOp",
 			"targets": [
 				{ "target": "MaxPool" },
-				{ "target": "AveragePool", "absolute_tolerance": 1e-6 }
+				{ "target": "AveragePool", "absolute_tolerance": 1e-6 },
+				{ "target": "LpPool", "absolute_tolerance": 1e-6 }
 			],
 			"runtimes": [
 				{ "name": "NNERuntimeRDGHlsl", "skip": true, "comment": "Not supported" }
@@ -531,7 +532,8 @@
 			"category": "GlobalPool",
 			"targets": [
 				{ "target": "GlobalMaxPool" },
-				{ "target": "GlobalAveragePool", "absolute_tolerance": 1e-6 }
+				{ "target": "GlobalAveragePool", "absolute_tolerance": 1e-6 },
+				{ "target": "GlobalLpPool", "absolute_tolerance": 1e-6 }
 			],
 			"runtimes": [
 				{ "name": "NNERuntimeRDGHlsl", "skip": true, "comment": "Not supported" }
diff --git a/Engine/Plugins/Experimental/NNE/Source/NNERuntimeRDG/Private/Dml/NNEDmlOperatorPool.cpp b/Engine/Plugins/Experimental/NNE/Source/NNERuntimeRDG/Private/Dml/NNEDmlOperatorPool.cpp
--- a/Engine/Plugins/Experimental/NNE/Source/NNERuntimeRDG/Private/Dml/NNEDmlOperatorPool.cpp
+++ b/Engine/Plugins/Experimental/NNE/Source/NNERuntimeRDG/Private/Dml/NNEDmlOperatorPool.cpp
@@ -188,6 +188,11 @@ public:
 			int32 IncludePadding = Attributes.GetValueOrDefault(TEXT("count_include_pad"), 0);
 			DmlPoolOpDesc.IncludePadding = static_cast<BOOL>(IncludePadding);
 		}
+
+		if constexpr (DmlOpType == DML_OPERATOR_LP_POOLING)
+		{
+			DmlPoolOpDesc.P = (UINT) Attributes.GetValueOrDefault<int>(TEXT("p"), 2);
+		}
 			
 		return CreateOperator(Device, DML_OPERATOR_DESC{ DmlOpType, &DmlPoolOpDesc });
 	}
@@ -209,6 +214,8 @@ NNE_DML_REGISTER_POOLING_OP(MaxPool, MAX, false)
 NNE_DML_REGISTER_POOLING_OP(GlobalMaxPool, MAX, true)
 NNE_DML_REGISTER_POOLING_OP(AveragePool, AVERAGE, false)
 NNE_DML_REGISTER_POOLING_OP(GlobalAveragePool, AVERAGE, true)
+NNE_DML_REGISTER_POOLING_OP(LpPool, LP, false)
+NNE_DML_REGISTER_POOLING_OP(GlobalLpPool, LP, true)
 
 #undef NNE_DML_REGISTER_POOLING_OP
 
