Add RDG DML Flatten operator

#rnx

diff --git a/Engine/Plugins/Experimental/NNE/Source/NNEQA/Resources/NNEQATestDesc.json b/Engine/Plugins/Experimental/NNE/Source/NNEQA/Resources/NNEQATestDesc.json
--- a/Engine/Plugins/Experimental/NNE/Source/NNEQA/Resources/NNEQATestDesc.json
+++ b/Engine/Plugins/Experimental/NNE/Source/NNEQA/Resources/NNEQATestDesc.json
@@ -400,9 +400,6 @@
 		},
 		{
 			"category": "Flatten",
-			"runtimes": [
-				{ "name": "NNERuntimeRDGDml", "skip": true },
-			],
 			"targets": [
 				{ "target": "Flatten" }
 			]
diff --git a/Engine/Plugins/Experimental/NNE/Source/NNERuntimeRDG/Private/Dml/NNEDmlOperatorFlatten.cpp b/Engine/Plugins/Experimental/NNE/Source/NNERuntimeRDG/Private/Dml/NNEDmlOperatorFlatten.cpp
new file mode 100644
--- /dev/null
+++ b/Engine/Plugins/Experimental/NNE/Source/NNERuntimeRDG/Private/Dml/NNEDmlOperatorFlatten.cpp
@@ -0,0 +1,78 @@
+// Copyright Epic Games, Inc. All Rights Reserved.
+
+#ifdef NNE_USE_DIRECTML
+
+#include "NNEDmlOperator.h"
+#include "NNEDmlOperatorUtils.h"
+
+namespace UE::NNERuntimeRDG::Private::Dml
+{
+
+class FOperatorDmlFlatten : public FOperatorDml
+{
+
+public:
+
+	static FOperatorDml* Create()
+	{
+		return new FOperatorDmlFlatten();
+	}
+
+	//
+	//
+	//
+	virtual bool Initialize(IDMLDevice* Device, TArrayView<const NNECore::Internal::FTensor> InputTensors, TArrayView<const NNECore::Internal::FTensor> OutputTensors, const NNECore::FAttributeMap& Attributes) override
+	{
+		check(InputTensors.Num() == 1);
+		check(OutputTensors.Num() == 1);
+
+		TConstArrayView<uint32>	InputShape = InputTensors[0].GetShape().GetData();
+		int32					InputRank = InputShape.Num();
+		int32					Axis = Attributes.GetValueOrDefault(TEXT("axis"), 1);
+		
+		if (Axis > InputRank || Axis < -InputRank)
+		{
+			UE_LOG(LogNNE, Warning, TEXT("Flatten 'Axis' attribute should be in the range [-r,r] with r being the rank of the input (name: %s) however axis is %d while rank is %d."), *InputTensors[0].GetName(), Axis, InputRank);
+			return false;
+		}
+
+		HandleNegativeAxis(Axis, InputRank);
+
+		uint32 InnerDimSize = 1;
+
+		for (int32 Idx = 0; Idx < Axis; ++Idx)
+		{
+			InnerDimSize *= InputShape[Idx];
+		}
+
+		DmlUtil::FSmallUIntArray	OutputShape;
+
+		OutputShape.Reserve(InputShape.Num());
+		OutputShape.Add(InnerDimSize);
+		OutputShape.Add(InputTensors[0].GetShape().Volume() / InnerDimSize);
+		
+		DmlUtil::FTensorDesc DmlTensorDesc;
+		
+		if (!DmlTensorDesc.InitFromTensor(OutputTensors[0], OutputShape.Num(),
+			/*Broadcast =*/ MakeEmptyArrayView<uint32>(),
+			/*CustomShape =*/ OutputShape))
+		{
+			UE_LOG(LogNNE, Warning, TEXT("Failed to initialize Flatten's output tensor for DML inference"));
+			return false;
+		}
+
+		DML_ELEMENT_WISE_IDENTITY_OPERATOR_DESC DmlIdentityOpDesc{};
+
+		DmlIdentityOpDesc.InputTensor = &DmlTensorDesc.Desc;
+		DmlIdentityOpDesc.OutputTensor = &DmlTensorDesc.Desc;
+
+		return CreateOperator(Device, DML_OPERATOR_DESC{ DML_OPERATOR_ELEMENT_WISE_IDENTITY, &DmlIdentityOpDesc });
+	}
+};
+
+// Register Reshape operator on Module startup
+NNE_DML_REGISTER_OP(Flatten)
+
+} // namespace UE::NNERuntimeRDG::Private::Dml
+
+#endif // NNE_USE_DIRECTML
