NNE HLSL Convolution matmul algo optim

Optimisations:
* 16 inputs in cache (allowing to reduce LDS and improve occupancy/unroll more)
* Epilogue for coalesced write
* Tranposed weights (when model support it) for coalesced read
-> Overall 1.5x faster (than before this CL)

Others:
* added support for convolution without bias for the algo

Tests:
* style transfer demo on win64 and consoles.
* all NNE tests green on win64 editor.

#rnx

diff --git a/Engine/Plugins/Experimental/NNE/Source/NNEHlslShaders/Internal/NNEHlslShadersConvCS.h b/Engine/Plugins/Experimental/NNE/Source/NNEHlslShaders/Internal/NNEHlslShadersConvCS.h
--- a/Engine/Plugins/Experimental/NNE/Source/NNEHlslShaders/Internal/NNEHlslShadersConvCS.h
+++ b/Engine/Plugins/Experimental/NNE/Source/NNEHlslShaders/Internal/NNEHlslShadersConvCS.h
@@ -45,11 +45,12 @@ namespace UE::NNEHlslShaders::Internal
 		SHADER_USE_PARAMETER_STRUCT(FConvCS, FGlobalShader)
 
 		class FConvAlgorithm : SHADER_PERMUTATION_ENUM_CLASS("ALGORITHM", EConvAlgorithm);
+		class FConvAreWeightsTransposed : SHADER_PERMUTATION_BOOL("WEIGHTS_TRANSPOSED");
 		class FConvGroupSize : SHADER_PERMUTATION_ENUM_CLASS("GROUP_SIZE", EConvGroupSize);
 		class FConvNumDimensions : SHADER_PERMUTATION_RANGE_INT("NUM_DIMENSIONS", 1, FConvConstants::MAX_NUM_DIMENSIONS);
 		class FConvNumReadsPerThread : SHADER_PERMUTATION_RANGE_INT("NUM_READS_PER_THREAD_POW2", FConvConstants::MIN_NUM_READS_PER_THREAD_POW2, FConvConstants::MAX_NUM_READS_PER_THREAD_POW2);
 		class FConvHasB : SHADER_PERMUTATION_BOOL("HAS_B");
-		using FPermutationDomain = TShaderPermutationDomain<FConvAlgorithm, FConvGroupSize, FConvNumDimensions, FConvNumReadsPerThread, FConvHasB>;
+		using FPermutationDomain = TShaderPermutationDomain<FConvAlgorithm, FConvAreWeightsTransposed, FConvGroupSize, FConvNumDimensions, FConvNumReadsPerThread, FConvHasB>;
 
 	public:
 		BEGIN_SHADER_PARAMETER_STRUCT(FParameters, )
@@ -63,6 +64,7 @@ namespace UE::NNEHlslShaders::Internal
 			SHADER_PARAMETER_ARRAY(FIntVector4, XBlockStartStride_XBlockStride_WDimension_WDimensionDilationXBlockStride, [FConvConstants::MAX_NUM_DIMENSIONS])
 			SHADER_PARAMETER_ARRAY(FVector4f, OneDiv_GroupStride_GroupThreadStride_XBlockStride, [FConvConstants::MAX_NUM_DIMENSIONS])
 			SHADER_PARAMETER(int32, NumWChannels)
+			SHADER_PARAMETER(int32, NumWFeatures)
 			SHADER_PARAMETER(int32, YBatchStride)
 			SHADER_PARAMETER(int32, YOutputKernelStride)
 			SHADER_PARAMETER(int32, XBatchStride)
diff --git a/Engine/Plugins/Experimental/NNE/Source/NNEHlslShaders/Internal/NNEHlslShadersConvMatmulCS.h b/Engine/Plugins/Experimental/NNE/Source/NNEHlslShaders/Internal/NNEHlslShadersConvMatmulCS.h
--- a/Engine/Plugins/Experimental/NNE/Source/NNEHlslShaders/Internal/NNEHlslShadersConvMatmulCS.h
+++ b/Engine/Plugins/Experimental/NNE/Source/NNEHlslShaders/Internal/NNEHlslShadersConvMatmulCS.h
@@ -13,6 +13,10 @@ namespace UE::NNEHlslShaders::Internal
 		DECLARE_GLOBAL_SHADER(FConvMatmulCS);
 		SHADER_USE_PARAMETER_STRUCT(FConvMatmulCS, FGlobalShader)
 
+		class FConvMatmulAreWeightsTransposed : SHADER_PERMUTATION_BOOL("WEIGHTS_TRANSPOSED");
+		class FConvMatmulHasBias : SHADER_PERMUTATION_BOOL("HAS_BIAS");
+		using FPermutationDomain = TShaderPermutationDomain<FConvMatmulAreWeightsTransposed, FConvMatmulHasBias>;
+
 	public:
 		BEGIN_SHADER_PARAMETER_STRUCT(FParameters, )
 			SHADER_PARAMETER_RDG_BUFFER_SRV(Buffer<float>, Input)
diff --git a/Engine/Plugins/Experimental/NNE/Source/NNEHlslShaders/Private/NNEHlslShadersConvCS.cpp b/Engine/Plugins/Experimental/NNE/Source/NNEHlslShaders/Private/NNEHlslShadersConvCS.cpp
--- a/Engine/Plugins/Experimental/NNE/Source/NNEHlslShaders/Private/NNEHlslShadersConvCS.cpp
+++ b/Engine/Plugins/Experimental/NNE/Source/NNEHlslShaders/Private/NNEHlslShadersConvCS.cpp
@@ -181,6 +181,7 @@ namespace UE::NNEHlslShaders::Internal
 			WChannelSize *= WShape[2 + i];
 		}
 		
+		Parameters.NumWFeatures = WShape[0];
 		Parameters.NumWChannels = WShape[1];
 
 		Parameters.YBatchStride = YShape[1] * YMemoryStride;
diff --git a/Engine/Plugins/Experimental/NNE/Source/NNEHlslShaders/Shaders/NNEHlslShadersConv.usf b/Engine/Plugins/Experimental/NNE/Source/NNEHlslShaders/Shaders/NNEHlslShadersConv.usf
--- a/Engine/Plugins/Experimental/NNE/Source/NNEHlslShaders/Shaders/NNEHlslShadersConv.usf
+++ b/Engine/Plugins/Experimental/NNE/Source/NNEHlslShaders/Shaders/NNEHlslShadersConv.usf
@@ -17,7 +17,7 @@
 #endif
 
 Buffer<BUFFER_TYPE> X; // N x C x XD0 x ... x XDi
-Buffer<BUFFER_TYPE> W; // M x C/GROUPS x WD0 x ... x WDi
+Buffer<BUFFER_TYPE> W; // M x C/GROUPS x WD0 x ... x WDi OR WD0 x WD1 x C X M if WEIGHTS_TRANSPOSED (only supported for 4D)
 RWBuffer<BUFFER_TYPE> Y;  // N x M x YD0 x ... x YDi
 Buffer<BUFFER_TYPE> B; // B0, ..., BM
 
@@ -47,7 +47,8 @@ int4 XBlockStartStride_XBlockStride_WDimension_WDimensionDilationXBlockStride[MA
 // z: 1/XBlockStride
 float4 OneDiv_GroupStride_GroupThreadStride_XBlockStride[MAX_NUM_DIMENSIONS];
 
-int NumWChannels; // WShape[1]
+int NumWFeatures; // Number of features (aka output channels)
+int NumWChannels; // Number of input channels
 
 int YBatchStride;
 int YOutputKernelStride;
@@ -201,9 +202,18 @@ void Conv(in const uint3 DispatchThreadID : SV_DispatchThreadID, in const uint3
         int WIndex = WChannelBatchOffset + GroupThreadID.x;
         if (GroupThreadID.x < WChannelBatchSize && WIndex < WOutputKernelStride)
         {
-            SharedMemoryW[GroupThreadID.x] = BUFFER_TO_WORK_TYPE(W[WChannelOutputKernelOffset + WIndex]);
+			#if WEIGHTS_TRANSPOSED
+				int MCHWIndex = WChannelOutputKernelOffset + WIndex;
+				int Whw = MCHWIndex % WChannelSize;
+				int Wc = (MCHWIndex / WChannelSize) % NumWChannels;
+				int Wm = (MCHWIndex / (NumWChannels * WChannelSize));
+				int HWCMIndex = Whw * NumWChannels * NumWFeatures + Wc * NumWFeatures + Wm;
+				SharedMemoryW[GroupThreadID.x] = BUFFER_TO_WORK_TYPE(W[HWCMIndex]);
+			#else
+				SharedMemoryW[GroupThreadID.x] = BUFFER_TO_WORK_TYPE(W[WChannelOutputKernelOffset + WIndex]);
+			#endif
             WChannelBatchOffset += WChannelBatchSize;
-        }
+		}
         
         // Each channel batch iterates over all consecutive elements of the loaded kernel channels
         // This results in bank conflict free lookups, as all threads of the same warp read exactly the same element which is then broadcasted
diff --git a/Engine/Plugins/Experimental/NNE/Source/NNEHlslShaders/Shaders/NNEHlslShadersConvMatmul.usf b/Engine/Plugins/Experimental/NNE/Source/NNEHlslShaders/Shaders/NNEHlslShadersConvMatmul.usf
--- a/Engine/Plugins/Experimental/NNE/Source/NNEHlslShaders/Shaders/NNEHlslShadersConvMatmul.usf
+++ b/Engine/Plugins/Experimental/NNE/Source/NNEHlslShaders/Shaders/NNEHlslShadersConvMatmul.usf
@@ -7,15 +7,19 @@
 #define BUFFER_TO_WORK_TYPE(x) x
 #define WORK_TO_BUFFER_TYPE(x) x
 
-#define LDS_IC 32 //Input channel cached in LDS
+#define LDS_IC 16 //Input channel cached in LDS
 #define TG_OC 32 //Number of output channel covered by a thread group
 #define TG_OP 32 //Number of output pixel covered by a thread group
 
-groupshared WORK_TYPE LDSInputs[LDS_IC][TG_OP]; //1024 WORK_TYPE
-groupshared WORK_TYPE LDSWeights[LDS_IC][TG_OC]; //1024 WORK_TYPE
+//To avoid bank conflict when writing out to DDR
+#define LDS_PADX 1
+#define LDS_PADY 2
+#define LDS_X_PADDEDSIZE ((4 + LDS_PADX) * 8 + LDS_PADY) //42
+
+groupshared WORK_TYPE LDSBuffer[32 * LDS_X_PADDEDSIZE]; //max (16*32+32,32*42) = 1344 WORK_TYPE 
 
 Buffer<BUFFER_TYPE> Input; // Ni x Ci x Hi x Wi
-Buffer<BUFFER_TYPE> Weight; // Cw x Ci x Hw x Ww
+Buffer<BUFFER_TYPE> Weight; // Cw x Ci x Hw x Ww OR Hw x Ww x Ci x Cw if WEIGHTS_TRANSPOSED
 Buffer<BUFFER_TYPE> Bias; // Cw
 RWBuffer<BUFFER_TYPE> Output; // Ni x Cw x Ho x Wo
 
@@ -32,7 +36,25 @@ int PadTop;
 int StrideH;
 int StrideW;
 
-//#define DEBUGVAL(val, pos, thread) {if (thread == GroupThreadIdx) Output[pos] = val;}
+WORK_TYPE ReadLDSWeights(int inputChannel, int outputChannel)
+{
+	return LDSBuffer[LDS_IC * TG_OP + inputChannel * TG_OC + outputChannel];
+}
+
+void WriteLDSWeights(int inputChannel, int outputChannel, WORK_TYPE value)
+{
+	LDSBuffer[LDS_IC * TG_OP + inputChannel * TG_OC + outputChannel] = value;
+}
+
+WORK_TYPE ReadLDSInputs(int inputChannel, int outputPixel)
+{
+	return LDSBuffer[inputChannel * TG_OP + outputPixel];
+}
+
+void WriteLDSInputs(int inputChannel, int outputPixel, WORK_TYPE value)
+{
+	LDSBuffer[inputChannel * TG_OP + outputPixel] = value;
+}
 
 [numthreads(8, 8, 1)]
 void ConvMatmul(
@@ -41,7 +63,6 @@ void ConvMatmul(
 	in const uint3 GroupThreadID : SV_GroupThreadID,
 	in const uint GroupIndex : SV_GroupIndex)
 {
-	
 	const int DispatchThreadOutputPixelOffset = 4 * DispatchThreadID.x;
 	const int DispatchThreadOutputChannelOffset = 4 * DispatchThreadID.y;
 	const int GroupThreadOutputPixelOffset = 4 * GroupThreadID.x;
@@ -59,10 +80,17 @@ void ConvMatmul(
 	const int cw = DispatchThreadOutputChannelOffset;
 	const int pi = DispatchThreadOutputPixelOffset;
 
+#if HAS_BIAS
 	WORK_TYPE BiasC0 = BUFFER_TO_WORK_TYPE(Bias[cw+0]);
 	WORK_TYPE BiasC1 = BUFFER_TO_WORK_TYPE(Bias[cw+1]);
 	WORK_TYPE BiasC2 = BUFFER_TO_WORK_TYPE(Bias[cw+2]);
 	WORK_TYPE BiasC3 = BUFFER_TO_WORK_TYPE(Bias[cw+3]);
+#else
+	WORK_TYPE BiasC0 = 0.0f;
+	WORK_TYPE BiasC1 = 0.0f;
+	WORK_TYPE BiasC2 = 0.0f;
+	WORK_TYPE BiasC3 = 0.0f;
+#endif
 	
 	int i;
 	// First index is thread output channel offset (on cw)
@@ -79,7 +107,7 @@ void ConvMatmul(
 	}
 	
 	int Scalar_KernelIdx = 0;
-	const int InputChannelOffset = GroupThreadIdx >= 32 ? 16 : 0;
+	const int InputChannelOffset = GroupThreadIdx >= 32 ? LDS_IC/2 : 0;
 	const int OutputPixelOffset = GroupThreadIdx % 32;
 	const int OutputChannelOffset = GroupThreadIdx % 32;
 	
@@ -99,25 +127,29 @@ void ConvMatmul(
 			for (int ci = 0; ci < Ci; ci += LDS_IC)
 			{
 				///— DDR → LDS
-				// We need to load 2048 floats
-				// 1024 weights (32 inputs channel, 32 output channels) and
-				// 1024 inputs (32 inputs channel, 32 inputs pixels)
-				// We have 64 threads thus each will read 16 inputs and 16 weights.
-				// the first 32 read the first 16 channels, while the next 32 read the remaining 16.
+				// We need to load 1024 floats
+				// 512 weights (16 inputs channel, 32 output channels) and
+				// 512 inputs (16 inputs channel, 32 inputs pixels)
+				// We have 64 threads thus each will read 8 inputs and 8 weights.
 
 				const int ReadOffsetI = Scalar_BatchInputOffset + (ci + InputChannelOffset) * Hi * Wi + InputPixelOffset;
-				const int ReadOffsetW = (Scalar_GroupOutputChannelOffset + OutputChannelOffset) * Ci * Hw * Ww + (ci + InputChannelOffset) * Hw * Ww + Scalar_KernelIdx;
+				#if WEIGHTS_TRANSPOSED
+					const int InputChannelStride = Cw;
+					const int ReadOffsetW = Scalar_KernelIdx + (ci + InputChannelOffset) * InputChannelStride + (Scalar_GroupOutputChannelOffset + OutputChannelOffset);
+				#else
+					const int InputChannelStride = Hw * Ww;
+					const int ReadOffsetW = (Scalar_GroupOutputChannelOffset + OutputChannelOffset) * Ci * InputChannelStride + (ci + InputChannelOffset) * InputChannelStride + Scalar_KernelIdx;
+				#endif
 				
 				UNROLL
 				for (i = 0; i < (LDS_IC / 2); ++i)
 				{
 					WORK_TYPE ValueI = BUFFER_TO_WORK_TYPE(Input[Hi * Wi * i + ReadOffsetI]);
 					ValueI = IsValidInput ? ValueI : 0.0f;
-					LDSInputs[InputChannelOffset + i][OutputPixelOffset] = ValueI;
+					WriteLDSInputs(InputChannelOffset + i, OutputPixelOffset, ValueI);
 
-					// Idea: Try to transpose weights to W[Kw, Ci, Cw] to get coalesced reads.
-					WORK_TYPE ValueW = BUFFER_TO_WORK_TYPE(Weight[Hw * Ww * i + ReadOffsetW]);
-					LDSWeights[InputChannelOffset + i][OutputChannelOffset] = ValueW;
+					WORK_TYPE ValueW = BUFFER_TO_WORK_TYPE(Weight[InputChannelStride * i + ReadOffsetW]);
+					WriteLDSWeights(InputChannelOffset + i, OutputChannelOffset, ValueW);
 				}
 				
 				GroupMemoryBarrierWithGroupSync();
@@ -137,8 +169,8 @@ void ConvMatmul(
 						UNROLL
 						for (i = 0; i < 4; ++i)
 						{
-							RegWeights[unrolledIdx][i] = LDSWeights[cachedCiBase + unrolledIdx][GroupThreadOutputChannelOffset + i];
-							RegInputs[unrolledIdx][i] = LDSInputs[cachedCiBase + unrolledIdx][GroupThreadOutputPixelOffset + i];
+							RegWeights[unrolledIdx][i] = ReadLDSWeights(cachedCiBase + unrolledIdx, GroupThreadOutputChannelOffset + i);
+							RegInputs[unrolledIdx][i] = ReadLDSInputs(cachedCiBase + unrolledIdx, GroupThreadOutputPixelOffset + i);
 						}
 
 						// Inner loop (16 mads)
@@ -155,18 +187,37 @@ void ConvMatmul(
 				
 				GroupMemoryBarrierWithGroupSync();
 			}
+			#if WEIGHTS_TRANSPOSED
+				Scalar_KernelIdx += Ci * Cw;
+			#else
 			++Scalar_KernelIdx;
+			#endif
 		}
 	}
 	
-	/// register to DDR
+	/// Write results back to DDR (via LDS)
 	UNROLL
 	for (i = 0; i < 4; ++i)
 	{
-		const int WriteOffset = Scalar_BatchOutputOffset + (cw + i) * Ho * Wo + pi;
-		Output[WriteOffset + 0] = Values[i][0];
-		Output[WriteOffset + 1] = Values[i][1];
-		Output[WriteOffset + 2] = Values[i][2];
-		Output[WriteOffset + 3] = Values[i][3];
+		const int LDSWriteOffsetY = 4 * GroupThreadID.y + i;
+		const int LDSWriteOffsetX = (LDS_PADX + 4) * GroupThreadID.x;
+		LDSBuffer[LDSWriteOffsetY * LDS_X_PADDEDSIZE + LDSWriteOffsetX + 0] = Values[i][0];
+		LDSBuffer[LDSWriteOffsetY * LDS_X_PADDEDSIZE + LDSWriteOffsetX + 1] = Values[i][1];
+		LDSBuffer[LDSWriteOffsetY * LDS_X_PADDEDSIZE + LDSWriteOffsetX + 2] = Values[i][2];
+		LDSBuffer[LDSWriteOffsetY * LDS_X_PADDEDSIZE + LDSWriteOffsetX + 3] = Values[i][3];
+	}
+
+	GroupMemoryBarrierWithGroupSync();
+
+	const int WriteOutChannelOffset = GroupThreadIdx >= 32 ? LDS_IC : 0;
+	const int WriteOutPixelOffset = GroupThreadIdx % 32;
+	const int LDSPixelOffset = (4 + LDS_PADX) * (WriteOutPixelOffset / 4) + WriteOutPixelOffset % 4;
+	
+	UNROLL
+	for (i = 0; i < LDS_IC; ++i)
+	{
+		const int WriteOffset = Scalar_BatchOutputOffset + (Scalar_GroupOutputChannelOffset + i + WriteOutChannelOffset) * Ho * Wo + Scalar_GroupOutputPixelOffset + WriteOutPixelOffset;
+		Output[WriteOffset] = LDSBuffer[(i + WriteOutChannelOffset) * LDS_X_PADDEDSIZE + LDSPixelOffset];
 	}
+
 }
\ No newline at end of file
diff --git a/Engine/Plugins/Experimental/NNE/Source/NNEQA/Private/NNEQAUnitTestOpTranspose.cpp b/Engine/Plugins/Experimental/NNE/Source/NNEQA/Private/NNEQAUnitTestOpTranspose.cpp
new file mode 100644
--- /dev/null
+++ b/Engine/Plugins/Experimental/NNE/Source/NNEQA/Private/NNEQAUnitTestOpTranspose.cpp
@@ -0,0 +1,102 @@
+// Copyright Epic Games, Inc. All Rights Reserved.
+
+#include "CoreMinimal.h"
+#include "NNECoreTensor.h"
+#include "NNECoreTypes.h"
+#include "NNEQAUnitTestHelper.h"
+#include "NNERuntimeRDGHelperTranspose.h"
+
+#if WITH_DEV_AUTOMATION_TESTS
+#include "Misc/AutomationTest.h"
+#endif
+
+namespace UE::NNEQA::Private::NNERuntimeRDG::TransposeOp
+{
+
+#if WITH_DEV_AUTOMATION_TESTS
+
+	using namespace NNECore::Internal;
+	using namespace UE::NNERuntimeRDG::Internal;
+
+	IMPLEMENT_NNE_SHAPEINFERENCEHELPER_UNIT_AUTOMATION_TEST(FTransposeCPUHelperIdentity, "System.Engine.MachineLearning.NNE.UnitTest.TransposeHelper.Identity");
+	bool FTransposeCPUHelperIdentity::RunTest(const FString& Parameter)
+	{
+		FTensor XC1x2 = MakeConstTensor(TEXT("XC1x2"), { 1,2 }, { 1.0f, 2.0f });
+		FTensor XC1x2x3 = MakeConstTensor(TEXT("XC1x2x3"), { 1,2,3 }, { 1.0f, 2.0f, 3.0f, 4.0f, 5.0f, 6.0f });
+
+		CPUHelper::Transpose::TransposePreparedData(XC1x2, { 0,1 });
+		UTEST_EQUAL(TEXT("Transpose(XC1x2,0,1)[0]"), XC1x2.GetPreparedData<float>()[0], 1.0f);
+		UTEST_EQUAL(TEXT("Transpose(XC1x2,0,1)[1]"), XC1x2.GetPreparedData<float>()[1], 2.0f);
+
+		CPUHelper::Transpose::TransposePreparedData(XC1x2x3, { 0,1,2 });
+		UTEST_EQUAL(TEXT("Transpose(XC1x2x3,0,1,2)[0]"), XC1x2x3.GetPreparedData<float>()[0], 1.0f);
+		UTEST_EQUAL(TEXT("Transpose(XC1x2x3,0,1,2)[1]"), XC1x2x3.GetPreparedData<float>()[1], 2.0f);
+		UTEST_EQUAL(TEXT("Transpose(XC1x2x3,0,1,2)[2]"), XC1x2x3.GetPreparedData<float>()[2], 3.0f);
+		UTEST_EQUAL(TEXT("Transpose(XC1x2x3,0,1,2)[3]"), XC1x2x3.GetPreparedData<float>()[3], 4.0f);
+		UTEST_EQUAL(TEXT("Transpose(XC1x2x3,0,1,2)[4]"), XC1x2x3.GetPreparedData<float>()[4], 5.0f);
+		UTEST_EQUAL(TEXT("Transpose(XC1x2x3,0,1,2)[5]"), XC1x2x3.GetPreparedData<float>()[5], 6.0f);
+
+		return true;
+	}
+
+	IMPLEMENT_NNE_SHAPEINFERENCEHELPER_UNIT_AUTOMATION_TEST(FTransposeCPUHelperReverse, "System.Engine.MachineLearning.NNE.UnitTest.TransposeHelper.Reverse");
+	bool FTransposeCPUHelperReverse::RunTest(const FString& Parameter)
+	{
+		FTensor XC1x2 = MakeConstTensor(TEXT("XC1x2"), { 1,2 }, { 1.0f, 2.0f });
+		FTensor XC1x2x3 = MakeConstTensor(TEXT("XC1x2x3"), { 1,2,3 }, { 1.0f, 2.0f, 3.0f, 4.0f, 5.0f, 6.0f });
+
+		CPUHelper::Transpose::TransposePreparedData(XC1x2, { 1,0 });
+		UTEST_EQUAL(TEXT("Transpose(XC1x2,1,0)[0]"), XC1x2.GetPreparedData<float>()[0], 1.0f);
+		UTEST_EQUAL(TEXT("Transpose(XC1x2,1,0)[1]"), XC1x2.GetPreparedData<float>()[1], 2.0f);
+
+		CPUHelper::Transpose::TransposePreparedData(XC1x2x3, { 2,1,0 });
+		UTEST_EQUAL(TEXT("Transpose(XC1x2x3,2,1,0)[0]"), XC1x2x3.GetPreparedData<float>()[0], 1.0f);
+		UTEST_EQUAL(TEXT("Transpose(XC1x2x3,2,1,0)[1]"), XC1x2x3.GetPreparedData<float>()[1], 4.0f);
+		UTEST_EQUAL(TEXT("Transpose(XC1x2x3,2,1,0)[2]"), XC1x2x3.GetPreparedData<float>()[2], 2.0f);
+		UTEST_EQUAL(TEXT("Transpose(XC1x2x3,2,1,0)[3]"), XC1x2x3.GetPreparedData<float>()[3], 5.0f);
+		UTEST_EQUAL(TEXT("Transpose(XC1x2x3,2,1,0)[4]"), XC1x2x3.GetPreparedData<float>()[4], 3.0f);
+		UTEST_EQUAL(TEXT("Transpose(XC1x2x3,2,1,0)[5]"), XC1x2x3.GetPreparedData<float>()[5], 6.0f);
+
+		return true;
+	}
+
+	IMPLEMENT_NNE_SHAPEINFERENCEHELPER_UNIT_AUTOMATION_TEST(FTransposeCPUHelperConvKernel, "System.Engine.MachineLearning.NNE.UnitTest.TransposeHelper.ConvKernel");
+	bool FTransposeCPUHelperConvKernel::RunTest(const FString& Parameter)
+	{
+		FTensor XC2x3x4x5 = MakeConstTensor(TEXT("XC2x3x4x5"), { 2,3,4,5 }, { 
+			0.0f,   1.0f,   2.0f,   3.0f,   4.0f,   5.0f,   6.0f,   7.0f,   8.0f,   9.0f,   10.0f,  11.0f,  12.0f,
+			13.0f,  14.0f,  15.0f,  16.0f,  17.0f,  18.0f,  19.0f,  20.0f,  21.0f,  22.0f,  23.0f,  24.0f,  25.0f,
+			26.0f,  27.0f,  28.0f,  29.0f,  30.0f,  31.0f,  32.0f,  33.0f,  34.0f,  35.0f,  36.0f,  37.0f,  38.0f,
+			39.0f,  40.0f,  41.0f,  42.0f,  43.0f,  44.0f,  45.0f,  46.0f,  47.0f,  48.0f,  49.0f,  50.0f,  51.0f,
+			52.0f,  53.0f,  54.0f,  55.0f,  56.0f,  57.0f,  58.0f,  59.0f,  60.0f,  61.0f,  62.0f,  63.0f,  64.0f,
+			65.0f,  66.0f,  67.0f,  68.0f,  69.0f,  70.0f,  71.0f,  72.0f,  73.0f,  74.0f,  75.0f,  76.0f,  77.0f,
+			78.0f,  79.0f,  80.0f,  81.0f,  82.0f,  83.0f,  84.0f,  85.0f,  86.0f,  87.0f,  88.0f,  89.0f,  90.0f,
+			91.0f,  92.0f,  93.0f,  94.0f,  95.0f,  96.0f,  97.0f,  98.0f,  99.0f, 100.0f, 101.0f, 102.0f, 103.0f,
+			104.0f,105.0f, 106.0f, 107.0f, 108.0f, 109.0f, 110.0f, 111.0f, 112.0f, 113.0f, 114.0f, 115.0f, 116.0f,
+			117.0f,118.0f, 119.0f});
+
+		TArray<float> ExpectedResult({
+			0.0f,   60.0f,  20.0f,  80.0f,  40.0f, 100.0f,   1.0f,  61.0f,  21.0f,  81.0f,  41.0f, 101.0f,   2.0f,
+			62.0f,  22.0f,  82.0f,  42.0f, 102.0f,   3.0f,  63.0f,  23.0f,  83.0f,  43.0f, 103.0f,   4.0f,  64.0f,
+			24.0f,  84.0f,  44.0f, 104.0f,   5.0f,  65.0f,  25.0f,  85.0f,  45.0f, 105.0f,   6.0f,  66.0f,  26.0f,
+			86.0f,  46.0f, 106.0f,   7.0f,  67.0f,  27.0f,  87.0f,  47.0f, 107.0f,   8.0f,  68.0f,  28.0f,  88.0f,
+			48.0f, 108.0f,   9.0f,  69.0f,  29.0f,  89.0f,  49.0f, 109.0f,  10.0f,  70.0f,  30.0f,  90.0f,  50.0f,
+			110.0f, 11.0f,  71.0f,  31.0f,  91.0f,  51.0f, 111.0f,  12.0f,  72.0f,  32.0f,  92.0f,  52.0f, 112.0f,
+			13.0f,  73.0f,  33.0f,  93.0f,  53.0f, 113.0f,  14.0f,  74.0f,  34.0f,  94.0f,  54.0f, 114.0f,  15.0f,
+			75.0f,  35.0f,  95.0f,  55.0f, 115.0f,  16.0f,  76.0f,  36.0f,  96.0f,  56.0f, 116.0f,  17.0f,  77.0f,
+			37.0f,  97.0f,  57.0f, 117.0f,  18.0f,  78.0f,  38.0f,  98.0f,  58.0f, 118.0f,  19.0f,  79.0f,  39.0f,
+			99.0f,  59.0f, 119 });
+
+		CPUHelper::Transpose::TransposePreparedData(XC2x3x4x5, { 2,3,1,0 });
+		for (uint32 i = 0; i < XC2x3x4x5.GetVolume(); ++i)
+		{
+			FString TestText = FString::Printf(TEXT("Transpose(XC2x3x4x5,2,3,1,0)[ %d]"), i);
+			UTEST_EQUAL(*TestText, XC2x3x4x5.GetPreparedData<float>()[i], ExpectedResult[i]);
+		}
+
+		return true;
+	}
+
+#endif //WITH_DEV_AUTOMATION_TESTS
+
+} // namespace UE::NNEQA::Private::NNERuntimeRDG::SliceOp
diff --git a/Engine/Plugins/Experimental/NNE/Source/NNERuntimeRDG/Internal/NNERuntimeRDGHelperTranspose.h b/Engine/Plugins/Experimental/NNE/Source/NNERuntimeRDG/Internal/NNERuntimeRDGHelperTranspose.h
new file mode 100644
--- /dev/null
+++ b/Engine/Plugins/Experimental/NNE/Source/NNERuntimeRDG/Internal/NNERuntimeRDGHelperTranspose.h
@@ -0,0 +1,12 @@
+// Copyright Epic Games, Inc. All Rights Reserved.
+
+#pragma once 
+
+#include "CoreMinimal.h"
+
+namespace UE::NNECore::Internal { class FTensor; }
+
+namespace UE::NNERuntimeRDG::Internal::CPUHelper::Transpose
+{
+	bool NNERUNTIMERDG_API TransposePreparedData(NNECore::Internal::FTensor& Tensor, TConstArrayView<int32> Perms);
+} // UE::NNERuntimeRDG::Internal::CPUHelper::Transpose
\ No newline at end of file
diff --git a/Engine/Plugins/Experimental/NNE/Source/NNERuntimeRDG/Private/Helper/NNERuntimeRDGHelperTranspose.cpp b/Engine/Plugins/Experimental/NNE/Source/NNERuntimeRDG/Private/Helper/NNERuntimeRDGHelperTranspose.cpp
new file mode 100644
--- /dev/null
+++ b/Engine/Plugins/Experimental/NNE/Source/NNERuntimeRDG/Private/Helper/NNERuntimeRDGHelperTranspose.cpp
@@ -0,0 +1,58 @@
+// Copyright Epic Games, Inc. All Rights Reserved.
+
+#include "NNERuntimeRDGHelperTranspose.h"
+#include "NNERuntimeRDGTensorIdxIterator.h"
+
+namespace UE::NNERuntimeRDG::Internal::CPUHelper::Transpose
+{
+	bool TransposePreparedData(NNECore::Internal::FTensor& Tensor, TConstArrayView<int32> Perms)
+	{
+		if (!Tensor.HasPreparedData())
+		{
+			return false;
+		}
+		if (Tensor.GetDataType() != ENNETensorDataType::Float)
+		{
+			return false;
+		}
+
+		TConstArrayView<uint32> TensorShape = Tensor.GetShape().GetData();
+		TConstArrayView<float> InputData = Tensor.GetPreparedData<float>();
+
+		check(Perms.Num() == TensorShape.Num());
+
+		TArray<uint32> TransposedShape;
+		TransposedShape.SetNumUninitialized(TensorShape.Num());
+		for (int32 i = 0; i < TensorShape.Num(); ++i)
+		{
+			TransposedShape[i] = TensorShape[Perms[i]];
+		}
+
+		TArray<float> TransposedData;
+		TransposedData.SetNumUninitialized(InputData.Num());
+		TArray<uint32> TransposedPosition;
+		TransposedPosition.SetNumUninitialized(TransposedShape.Num());
+
+		Private::TensorIdxIterator it(Tensor.GetShape());
+		const Private::TensorIdxIterator itTransposed(NNECore::FTensorShape::Make(TransposedShape));
+
+		do
+		{
+			TConstArrayView<uint32> Position = it.GetPositions();
+			for (int32 i = 0; i < TensorShape.Num(); ++i)
+			{
+				TransposedPosition[i] = Position[Perms[i]];
+			}
+
+			float Value = InputData[it.GetIndex()];
+			int TransposedIndex = itTransposed.GetIndexFromPosition(TransposedPosition);
+			TransposedData[TransposedIndex] = Value;
+
+		} while (it.Advance());
+
+		Tensor.SetPreparedData<float>(TransposedData);
+
+		return true;
+	}
+	
+} // UE::NNERuntimeRDG::Private::CPUHelper::Transpose
diff --git a/Engine/Plugins/Experimental/NNE/Source/NNERuntimeRDG/Private/Hlsl/NNERuntimeRDGConv.cpp b/Engine/Plugins/Experimental/NNE/Source/NNERuntimeRDG/Private/Hlsl/NNERuntimeRDGConv.cpp
--- a/Engine/Plugins/Experimental/NNE/Source/NNERuntimeRDG/Private/Hlsl/NNERuntimeRDGConv.cpp
+++ b/Engine/Plugins/Experimental/NNE/Source/NNERuntimeRDG/Private/Hlsl/NNERuntimeRDGConv.cpp
@@ -2,12 +2,13 @@
 
 #include "NNERuntimeRDGConv.h"
 #include "Algo/AnyOf.h"
-#include "NNEHlslShadersConvCS.h"
-#include "NNEHlslShadersConvMatmulCS.h"
 #include "NNECoreAttributeMap.h"
-#include "NNERuntimeRDGHlslHelper.h"
 #include "NNECoreTensor.h"
 #include "NNECoreTypes.h"
+#include "NNEHlslShadersConvCS.h"
+#include "NNEHlslShadersConvMatmulCS.h"
+#include "NNERuntimeRDGHelperTranspose.h"
+#include "NNERuntimeRDGHlslHelper.h"
 
 namespace UE::NNERuntimeRDG::Private::Hlsl
 {
@@ -41,6 +42,7 @@ namespace UE::NNERuntimeRDG::Private::Hlsl
 		int32 Group = 1;
 		TArray<int32> Pads;
 		TArray<int32> Strides;
+		bool bAreWeightsTransposed = false;
 
 	public:
 
@@ -133,6 +135,7 @@ namespace UE::NNERuntimeRDG::Private::Hlsl
 			FConvCS::FPermutationDomain PermutationVector;
 
 			PermutationVector.Set<FConvCS::FConvAlgorithm>(Algorithm);
+			PermutationVector.Set<FConvCS::FConvAreWeightsTransposed>(bAreWeightsTransposed);
 			PermutationVector.Set<FConvCS::FConvGroupSize>(GroupSize);
 			PermutationVector.Set<FConvCS::FConvNumDimensions>(NumDimensions);
 			PermutationVector.Set<FConvCS::FConvNumReadsPerThread>(FConvCS::GetNumReadsPerThread(GroupSize, Weights.GetShape().GetData(), Dilations, Strides));
@@ -163,10 +166,8 @@ namespace UE::NNERuntimeRDG::Private::Hlsl
 				return false;
 			if (Weights.GetShape().Rank() != 4)
 				return false;
-			if (Bias == nullptr)
-				return false;
-			if (Bias->GetShape().Rank() != 1)
-				return false;
+			
+			bool bHasBias = Bias != nullptr;
 
 			const bool bHasDilation = Algo::AnyOf(Dilations, [](auto Dim) {return Dim != 1u; });
 			if (bHasDilation)
@@ -187,13 +188,11 @@ namespace UE::NNERuntimeRDG::Private::Hlsl
 			int Hw = Weights.GetShape().GetData()[2];
 			int Ww = Weights.GetShape().GetData()[3];
 
-			check(Cw == Bias->GetShape().GetData()[0]);
-
-			if (Ci % 32 != 0)
+			if (Ci % 16 != 0)
 				return false;
 			if (Cw % 32 != 0)
 				return false;
-			if (Wo % 32 != 0)
+			if (Wo % 32 != 0) // Idea : support this by launching more threads and discard some results so a threadgroup still operator on only one value for H.
 				return false;
 
 			TArray<int32> Padding = FConvCS::GetPadding(Input.GetShape().GetData(), Weights.GetShape().GetData(), AutoPad, Dilations, Strides, Pads);
@@ -202,7 +201,9 @@ namespace UE::NNERuntimeRDG::Private::Hlsl
 			FConvMatmulCS::FParameters* Params = GraphBuilder.AllocParameters<FConvMatmulCS::FParameters>();
 			Params->Input = GraphBuilder.CreateSRV(FRDGBufferSRVDesc(Input.GetBuffer(), PF_R32_FLOAT));
 			Params->Weight= GraphBuilder.CreateSRV(FRDGBufferSRVDesc(Weights.GetBuffer(), PF_R32_FLOAT));
-			Params->Bias = GraphBuilder.CreateSRV(FRDGBufferSRVDesc(Bias->GetBuffer(), PF_R32_FLOAT));
+			if (bHasBias) {
+				Params->Bias = GraphBuilder.CreateSRV(FRDGBufferSRVDesc(Bias->GetBuffer(), PF_R32_FLOAT));
+			}
 			Params->Output = GraphBuilder.CreateUAV(FRDGBufferUAVDesc(Output.GetBuffer(), PF_R32_FLOAT));
 			Params->Ci = Ci;
 			Params->Hi = Hi;
@@ -218,6 +219,8 @@ namespace UE::NNERuntimeRDG::Private::Hlsl
 			Params->PadLeft = Padding[1];
 
 			FConvMatmulCS::FPermutationDomain PermutationVector;
+			PermutationVector.Set<FConvMatmulCS::FConvMatmulAreWeightsTransposed>(bAreWeightsTransposed);
+			PermutationVector.Set<FConvMatmulCS::FConvMatmulHasBias>(bHasBias);
 			TShaderMapRef<FConvMatmulCS> ComputeShader(GetGlobalShaderMap(GMaxRHIFeatureLevel), PermutationVector);
 
 			RDG_EVENT_SCOPE(GraphBuilder, "NNE.Operator.Hlsl.Conv.Matmul");
@@ -263,6 +266,36 @@ namespace UE::NNERuntimeRDG::Private::Hlsl
 			}
 			DispatchConvDefault(GraphBuilder, Input, Weights, Bias, Output);
 		}
+
+		virtual void OptimizeInputsWeights(TArrayView<FTensorRDGRef> InputWeights) override
+		{
+			using namespace UE::NNEHlslShaders::Internal;
+
+			check(InputWeights.Num() >= 2);
+			FTensorRDGRef Weights = InputWeights[1];
+			if (Weights == nullptr)
+				return;
+
+			// Heuristics : only matmul implementation benefits from transposed weights
+			if (Weights->GetShape().Rank() != 4)
+				return;
+
+			uint32 Cw = Weights->GetShape().GetData()[0];
+			uint32 Ci = Weights->GetShape().GetData()[1];
+
+			if (Ci % 16 != 0)
+				return;
+			if (Cw % 32 != 0)
+				return;
+			if (Group != 1)
+				return;
+
+			//Transpose from CwCiHwWw to HwWwCiCw
+			if (Internal::CPUHelper::Transpose::TransposePreparedData(*Weights, {2,3,1,0} ))
+			{
+				bAreWeightsTransposed = true;
+			}
+		};
 	};
 
 	bool ValidateConvOperator(const NNECore::FAttributeMap& AttributeMap, TConstArrayView<ENNETensorDataType> InputTypes, TConstArrayView<NNECore::FSymbolicTensorShape> InputShapes)
diff --git a/Engine/Plugins/Experimental/NNE/Source/NNERuntimeRDG/Private/NNERuntimeRDGHlslOp.h b/Engine/Plugins/Experimental/NNE/Source/NNERuntimeRDG/Private/NNERuntimeRDGHlslOp.h
--- a/Engine/Plugins/Experimental/NNE/Source/NNERuntimeRDG/Private/NNERuntimeRDGHlslOp.h
+++ b/Engine/Plugins/Experimental/NNE/Source/NNERuntimeRDG/Private/NNERuntimeRDGHlslOp.h
@@ -14,6 +14,7 @@ namespace UE::NNERuntimeRDG::Private::Hlsl
 		virtual ~FOperatorHlsl() = default;
 		virtual bool Initialize(TConstArrayView<NNECore::FTensorDesc> InputTensorDescs, TConstArrayView<NNECore::FTensorDesc> OutputTensorDescs, const NNECore::FAttributeMap& Attributes) = 0;
 		virtual void Dispatch(FRDGBuilder& GraphBuilder, TConstArrayView<FTensorRDGRef> InInputTensors, TConstArrayView<FTensorRDGRef> InOutputTensors) = 0;
+		virtual void OptimizeInputsWeights(TArrayView<FTensorRDGRef> InputWeights) {};
 	};
 
 	/**
diff --git a/Engine/Plugins/Experimental/NNE/Source/NNERuntimeRDG/Private/NNERuntimeRDGModelHlsl.cpp b/Engine/Plugins/Experimental/NNE/Source/NNERuntimeRDG/Private/NNERuntimeRDGModelHlsl.cpp
--- a/Engine/Plugins/Experimental/NNE/Source/NNERuntimeRDG/Private/NNERuntimeRDGModelHlsl.cpp
+++ b/Engine/Plugins/Experimental/NNE/Source/NNERuntimeRDG/Private/NNERuntimeRDGModelHlsl.cpp
@@ -35,6 +35,29 @@ FOperatorHlsl* OpCreate(const FString& OpName, TConstArrayView<NNECore::FTensorD
 	return Op;
 }
 
+FTensorRDGRef GetTensorRefIfWeightAndOnlyUsedByOperator(int32 TensorIndex, int32 OperatorIdx,
+	TConstArrayView<FNNEFormatOperatorDesc> Operators, TConstArrayView<int32> WeightTensorIndices, FTensorRDGArray& WeightTensorRDGs)
+{
+	for (int32 Idx = 0; Idx < Operators.Num(); ++Idx)
+	{
+		if (Idx == OperatorIdx)
+			continue;
+		for (int32 InputTensorIndex : Operators[Idx].InTensors)
+		{
+			if (InputTensorIndex == TensorIndex)
+				return nullptr;
+		}
+	}
+
+	int32 WeightIndex;
+	if (WeightTensorIndices.Find(TensorIndex, WeightIndex))
+	{
+		return &WeightTensorRDGs[WeightIndex];
+	}
+
+	return nullptr;
+}
+
 } // namespace ModelUtils
 
 bool FModel::PrepareModelRDG(FRDGBuilder& RDGBuilder)
@@ -93,21 +116,25 @@ bool FModel::Init(TConstArrayView<uint8> ModelData)
 		return false;
 	}
 
-	// Create HLSL tensor and upload to GPU
-	PrepareWeights();
+	TArray<NNECore::FTensorDesc> Inputs;
+	TArray<NNECore::FTensorDesc> Outputs;
+	TArray<FTensorRDGRef> InputsAsWeights;
 
 	// Loop over all operators in the model and create them
 	for (int32 Idx = 0; Idx < Format.Operators.Num(); ++Idx)
 	{
 		const FString TypeName = Format.Operators[Idx].TypeName;
-
-		TArray<NNECore::FTensorDesc> Inputs;
-		TArray<NNECore::FTensorDesc> Outputs;
+		Inputs.Reset();
+		Outputs.Reset();
+		InputsAsWeights.Reset();
 		UE::NNECore::FAttributeMap AttributeMap;
 
 		for (int32 InputTensorIndex : Format.Operators[Idx].InTensors)
 		{
+			FTensorRDGRef InputAsWeight = ModelUtils::GetTensorRefIfWeightAndOnlyUsedByOperator(InputTensorIndex, Idx, Format.Operators, WeightTensorIndices, WeightTensorRDGs);
+
 			Inputs.Emplace(AllSymbolicTensorDescs[InputTensorIndex]);
+			InputsAsWeights.Emplace(InputAsWeight);
 		}
 		for (int32 OutputTensorIndex : Format.Operators[Idx].OutTensors)
 		{
@@ -128,9 +155,14 @@ bool FModel::Init(TConstArrayView<uint8> ModelData)
 			return false;
 		}
 
+		Op->OptimizeInputsWeights(InputsAsWeights);
+
 		Operators.Add(Op);
 	}
 
+	// Create HLSL tensor and upload to GPU
+	PrepareWeights();
+
 	return true;
 }
 
