Add RDG DirectML Dropout operator

#rnx

diff --git a/Engine/Plugins/Experimental/NNE/Source/NNEQA/Resources/NNEQATestDesc.json b/Engine/Plugins/Experimental/NNE/Source/NNEQA/Resources/NNEQATestDesc.json
--- a/Engine/Plugins/Experimental/NNE/Source/NNEQA/Resources/NNEQATestDesc.json
+++ b/Engine/Plugins/Experimental/NNE/Source/NNEQA/Resources/NNEQATestDesc.json
@@ -391,9 +391,6 @@
 		},
 		{
 			"category": "Dropout",
-			"runtimes": [
-				{ "name": "NNERuntimeRDGDml", "skip": true },
-			],
 			"targets": [
 				{ "target": "Dropout" }
 			]
diff --git a/Engine/Plugins/Experimental/NNE/Source/NNERuntimeRDG/Private/Dml/NNEDmlOperatorActivations.cpp b/Engine/Plugins/Experimental/NNE/Source/NNERuntimeRDG/Private/Dml/NNEDmlOperatorActivations.cpp
--- a/Engine/Plugins/Experimental/NNE/Source/NNERuntimeRDG/Private/Dml/NNEDmlOperatorActivations.cpp
+++ b/Engine/Plugins/Experimental/NNE/Source/NNERuntimeRDG/Private/Dml/NNEDmlOperatorActivations.cpp
@@ -245,6 +245,7 @@ REGISTER_OP_ACTIVATION_UNARY(OpName, DmlOpName)
 // 	REGISTER_OP_ACTIVATION_UNARY(			OpName, 			DmlOpName )
 // 	REGISTER_OP_ACTIVATION_UNARY_PARAMS(	OpName, 			DmlOpName, 				InAlpha, 							InBeta, 	InGamma )
 
+	REGISTER_OP_ACTIVATION_UNARY(			Dropout,			IDENTITY )
 	REGISTER_OP_ACTIVATION_UNARY_PARAMS(	Elu, 				ELU, 					1.0f, 								0.0f, 		1.05070102214813232421875f )
 	REGISTER_OP_ACTIVATION_UNARY_PARAMS(	HardSigmoid, 		HARD_SIGMOID, 			0.2f, 								0.5f, 		0.0f )
 	REGISTER_OP_ACTIVATION_UNARY_PARAMS(	LeakyRelu, 			LEAKY_RELU, 			0.01f, 								0.0f, 		0.0f )
