NNE - Add Rdg Dml Split

#rnx

diff --git a/Engine/Plugins/Experimental/NNE/Source/NNEQA/Resources/NNEQATestDesc.json b/Engine/Plugins/Experimental/NNE/Source/NNEQA/Resources/NNEQATestDesc.json
--- a/Engine/Plugins/Experimental/NNE/Source/NNEQA/Resources/NNEQATestDesc.json
+++ b/Engine/Plugins/Experimental/NNE/Source/NNEQA/Resources/NNEQATestDesc.json
@@ -532,6 +532,15 @@
 			"runtimes": [
 				{ "name": "NNERuntimeRDGHlsl", "skip": true, "comment": "Not supported" }
 			]
+		},
+		{
+			"category": "Split",
+			"targets": [
+				{ "target": "Split" }
+			],
+			"runtimes": [
+				{ "name": "NNERuntimeRDGHlsl", "skip": true, "comment": "Not supported" }
+			]
 		}
 	],
 
@@ -1193,6 +1202,18 @@
 				{ "inputs": [ { "shape": [ 2, 3, 128, 256, 64 ] } ], "outputs": [ { "shape": [ 2, 3, 1, 1, 1 ] } ] }
 			]
 		},
+		{
+			"name": "Split.Axis2",
+			"datasets": [
+				{ "inputs": [ { "shape": [ 1, 1, 6, 2 ] } ], "outputs": [ { "shape": [ 1, 1, 2, 2 ] }, { "shape": [ 1, 1, 1, 2 ] }, { "shape": [ 1, 1, 3, 2 ] } ] }
+			]
+		},
+		{
+			"name": "Split.Axis3",
+			"datasets": [
+				{ "inputs": [ { "shape": [ 1, 1, 6, 2 ] } ], "outputs": [ { "shape": [ 1, 1, 6, 1 ] }, { "shape": [ 1, 1, 6, 1 ] } ] }
+			]
+		},
 		{ 
 			"name": "Reduce.Default",
 			"datasets": [
@@ -1720,6 +1741,26 @@
 				}
 			]
 		},
+		{
+			"name": "Split.Axis2",
+			"attribute_maps": [
+				{ "attributes": [ 
+						{ "name": "axis", "value": 2, "type": "Int32" },
+						{ "name": "split", "value": [2, 1, 3], "type": "Int32Array" }
+					]
+				}
+			]
+		},
+		{
+			"name": "Split.Axis3",
+			"attribute_maps": [
+				{ "attributes": [ 
+						{ "name": "axis", "value": 3, "type": "Int32" },
+						{ "name": "split", "value": [1, 1], "type": "Int32Array" }
+					]
+				}
+			]
+		},
 		{
 			"name": "Reduce.Axes.Keepdims",
 			"attribute_maps": [
diff --git a/Engine/Plugins/Experimental/NNE/Source/NNERuntimeRDG/Private/Dml/NNEDmlOperatorSplit.cpp b/Engine/Plugins/Experimental/NNE/Source/NNERuntimeRDG/Private/Dml/NNEDmlOperatorSplit.cpp
new file mode 100644
--- /dev/null
+++ b/Engine/Plugins/Experimental/NNE/Source/NNERuntimeRDG/Private/Dml/NNEDmlOperatorSplit.cpp
@@ -0,0 +1,107 @@
+// Copyright Epic Games, Inc. All Rights Reserved.
+
+#ifdef NNE_USE_DIRECTML
+#include "NNEDmlOperator.h"
+
+namespace UE::NNERuntimeRDG::Private::Dml
+{
+
+class FOperatorDmlSplit : public FOperatorDml
+{
+
+
+public:
+
+	static FOperatorDml* Create()
+	{
+		return new FOperatorDmlSplit();
+	}
+
+	//
+	//
+	//
+	virtual bool Initialize(IDMLDevice* Device, TArrayView<const NNECore::Internal::FTensor> InputTensors, TArrayView<const NNECore::Internal::FTensor> OutputTensors, const NNECore::FAttributeMap& Attributes) override
+	{
+		check(InputTensors.Num() == 1);
+		const NNECore::Internal::FTensor& InputTensor = InputTensors[0];
+
+		int Axis = Attributes.GetValueOrDefault<int>(TEXT("axis"), 0);
+
+		// Check split size is correct
+		//TODO: code should be validated by a validator which should report to the user if the contract is broken. We do it here instead.
+
+		uint32 SplitSize = 0;
+
+		for (int Idx = 0; Idx < OutputTensors.Num(); ++Idx)
+		{
+			//check(OutputTensors[Idx].GetShape().Rank() == InputTensor.GetShape().Rank());
+			if(OutputTensors[Idx].GetShape().Rank() != InputTensor.GetShape().Rank())
+			{
+				UE_LOG(LogNNE, Error, TEXT("Rank of output tensor and input tensor should be the same"));
+				return false;
+			}
+			for (int Dim = 0; Dim < InputTensor.GetShape().Rank(); ++Dim)
+			{
+				if (Dim == Axis)
+				{
+					SplitSize += OutputTensors[Idx].GetShape().GetData()[Dim];
+				}
+				else
+				{
+					//check(OutputTensors[Idx].GetShape().GetData()[Dim] == InputTensor.GetShape().GetData()[Dim]);
+					if(OutputTensors[Idx].GetShape().GetData()[Dim] != InputTensor.GetShape().GetData()[Dim])
+					{
+						UE_LOG(LogNNE, Error, TEXT("%s"), *FString::Printf(TEXT("Output tensor %d 's dimension %d should match input tensor's"), Idx, Dim));
+						return false;
+					}
+				}
+			}
+		}
+
+		//check(SplitSize == InputTensor.GetShape().GetData()[Axis]);
+		if(SplitSize != InputTensor.GetShape().GetData()[Axis])
+		{
+			UE_LOG(LogNNE, Error, TEXT("Input tensor's axis dimension size must be equal to the sum of output tensors' axis dimensions sizes"));
+			return false;
+		}
+
+
+		DmlUtil::FTensorDesc InputTensorDesc;
+		if (!InputTensorDesc.InitFromTensor(InputTensor, InputTensor.GetShape().Rank()))
+		{
+			UE_LOG(LogNNE, Error, TEXT("Failed to initialize pooling operator's input tensor for DML inference"));
+			return false;
+		}
+
+		DmlUtil::FSmallArray<DmlUtil::FTensorDesc> OutputTensorDescs;
+		DmlUtil::FSmallArray<DML_TENSOR_DESC> DmlOutputTensorDescs;
+		OutputTensorDescs.SetNum(OutputTensors.Num());
+		DmlOutputTensorDescs.SetNumUninitialized(OutputTensors.Num());
+
+		for(int Idx = 0; Idx < OutputTensors.Num(); ++Idx)
+		{
+			if (!OutputTensorDescs[Idx].InitFromTensor(OutputTensors[Idx], OutputTensors[Idx].GetShape().Rank()))
+			{
+				UE_LOG(LogNNE, Error, TEXT("Failed to initialize pooling operator's output tensor for DML inference"));
+				return false;
+			}
+			DmlOutputTensorDescs[Idx] = OutputTensorDescs[Idx].Desc;
+		}
+		
+		DML_SPLIT_OPERATOR_DESC	SplitOpDesc{};
+
+		SplitOpDesc.InputTensor = &InputTensorDesc.Desc;
+		SplitOpDesc.OutputCount = DmlOutputTensorDescs.Num();
+		SplitOpDesc.OutputTensors = DmlOutputTensorDescs.GetData();
+		SplitOpDesc.Axis = (UINT) Axis;
+
+		return CreateOperator(Device, DML_OPERATOR_DESC { DML_OPERATOR_SPLIT, &SplitOpDesc });
+	}
+};
+
+// Register Split operator on Module startup
+NNE_DML_REGISTER_OP(Split)
+
+} // namespace UE::NNERuntimeRDG::Private::Dml
+
+#endif // NNE_USE_DIRECTML
diff --git a/Engine/Plugins/Experimental/NNE/Source/NNEUtils/Private/NNEUtilsModelBuilderONNX.cpp b/Engine/Plugins/Experimental/NNE/Source/NNEUtils/Private/NNEUtilsModelBuilderONNX.cpp
--- a/Engine/Plugins/Experimental/NNE/Source/NNEUtils/Private/NNEUtilsModelBuilderONNX.cpp
+++ b/Engine/Plugins/Experimental/NNE/Source/NNEUtils/Private/NNEUtilsModelBuilderONNX.cpp
@@ -354,6 +354,7 @@ NNEUTILS_API bool CreateONNXModelForOperator(bool UseVariadicShapeForModel, cons
 	if (OperatorName == TEXT("BatchNormalization") ||	// current implementation is opset 9 (next version is 14)
 		OperatorName == TEXT("Clip") ||					// current implementation is opset 6 (next version is 11)
 		OperatorName == TEXT("Pad") ||					// current implementation is opset 2 (next version is 11)
+		OperatorName == TEXT("Split") ||				// current implementation is opset 2 (next version is 11)
 		OperatorName == TEXT("Shape") ||				// current implementation is opset 1 (next version is 13)
 		OperatorName == TEXT("Slice") ||				// current implementation is opset 1 (next version is 10)
 		OperatorName == TEXT("Squeeze") ||				// current implementation is opset 1 (next version is 11)
