NNE - Add RDG Dml LogSoftmax

#rnx

diff --git a/Engine/Plugins/Experimental/NNE/Source/NNEQA/Resources/NNEQATestDesc.json b/Engine/Plugins/Experimental/NNE/Source/NNEQA/Resources/NNEQATestDesc.json
--- a/Engine/Plugins/Experimental/NNE/Source/NNEQA/Resources/NNEQATestDesc.json
+++ b/Engine/Plugins/Experimental/NNE/Source/NNEQA/Resources/NNEQATestDesc.json
@@ -254,6 +254,10 @@
 						{ "name": "NNERuntimeRDGDml"},
 						{ "name": "NNERuntimeRDGHlsl"}
 					], "absolute_tolerance": 1e-6 },
+				{	"target": "LogSoftmax", 
+					"runtimes": [{ "name": "NNERuntimeRDGHlsl",   "skip": true }],
+					"additional_datasets": [ "Elementwise.Softmax.AttributeTest" ]
+				},
 				{	"target": "Neg" },
 				{	"target": "Not", "input_type": "boolean", "output_type": "boolean", "runtimes": [
 						{ "name": "NNERuntimeRDGDml", "skip": true, "comment": "RDG need bool tensor support"},
diff --git a/Engine/Plugins/Experimental/NNE/Source/NNERuntimeRDG/Private/Dml/NNEDmlOperatorActivations.cpp b/Engine/Plugins/Experimental/NNE/Source/NNERuntimeRDG/Private/Dml/NNEDmlOperatorActivations.cpp
--- a/Engine/Plugins/Experimental/NNE/Source/NNERuntimeRDG/Private/Dml/NNEDmlOperatorActivations.cpp
+++ b/Engine/Plugins/Experimental/NNE/Source/NNERuntimeRDG/Private/Dml/NNEDmlOperatorActivations.cpp
@@ -43,7 +43,8 @@ public:
 		const NNECore::Internal::FTensor& InputTensorDesc = InputTensors[0];
 		const NNECore::Internal::FTensor& OutputTensorDesc = OutputTensors[0];
 
-		if constexpr (std::is_same_v<DmlActivationOpDescType, DML_ACTIVATION_SOFTMAX1_OPERATOR_DESC>)
+		if constexpr (std::is_same_v<DmlActivationOpDescType, DML_ACTIVATION_SOFTMAX1_OPERATOR_DESC> ||
+					  std::is_same_v<DmlActivationOpDescType, DML_ACTIVATION_LOG_SOFTMAX1_OPERATOR_DESC>)
 		{
 			Axis = Attributes.GetValueOrDefault(TEXT("axis"), -1);
 		}
@@ -84,6 +85,14 @@ private:
 		Desc.OutputTensor = &TensorDesc.Desc;
 	}
 
+	void InitDmlOpDesc(DML_ACTIVATION_LOG_SOFTMAX1_OPERATOR_DESC& Desc, DmlUtil::FTensorDesc& TensorDesc)
+	{
+		Desc.InputTensor = &TensorDesc.Desc;
+		Desc.OutputTensor = &TensorDesc.Desc;
+		Desc.AxisCount = 1;
+		Desc.Axes = (UINT*) &Axis;
+	}
+
 	void InitDmlOpDesc(DML_ACTIVATION_SOFTMAX1_OPERATOR_DESC& Desc, DmlUtil::FTensorDesc& TensorDesc)
 	{
 		Desc.InputTensor = &TensorDesc.Desc;
@@ -249,6 +258,7 @@ REGISTER_OP_ACTIVATION_UNARY(OpName, DmlOpName)
 	REGISTER_OP_ACTIVATION_UNARY_PARAMS(	Elu, 				ELU, 					1.0f, 								0.0f, 		1.05070102214813232421875f )
 	REGISTER_OP_ACTIVATION_UNARY_PARAMS(	HardSigmoid, 		HARD_SIGMOID, 			0.2f, 								0.5f, 		0.0f )
 	REGISTER_OP_ACTIVATION_UNARY_PARAMS(	LeakyRelu, 			LEAKY_RELU, 			0.01f, 								0.0f, 		0.0f )
+	REGISTER_OP_ACTIVATION_UNARY(			LogSoftmax, 		LOG_SOFTMAX1 )
 	REGISTER_OP_ACTIVATION_UNARY(			Relu, 				RELU )
 	REGISTER_OP_ACTIVATION_UNARY_PARAMS(	Selu, 				SCALED_ELU, 			1.67326319217681884765625f, 		0.0f, 		1.05070102214813232421875f )
 	REGISTER_OP_ACTIVATION_UNARY(			Sigmoid, 			SIGMOID	)
