NNE Fix ORTDml interface

Tests: NNE smoke test green on win64 editor

#rnx

diff --git a/Engine/Plugins/Experimental/NNE/Source/NNECore/Public/NNECoreRuntimeGPU.h b/Engine/Plugins/Experimental/NNE/Source/NNECore/Public/NNECoreRuntimeGPU.h
new file mode 100644
--- /dev/null
+++ b/Engine/Plugins/Experimental/NNE/Source/NNECore/Public/NNECoreRuntimeGPU.h
@@ -0,0 +1,137 @@
+// Copyright Epic Games, Inc. All Rights Reserved.
+
+#pragma once
+
+#include "CoreMinimal.h"
+#include "NNECoreRuntime.h"
+#include "NNECoreTypes.h"
+#include "UObject/Interface.h"
+
+#include "NNECoreRuntimeGPU.generated.h"
+
+class UNNEModelData;
+
+namespace UE::NNECore
+{
+
+/**
+ * The tensor binding for passing input and output to IModelGPU.
+ *
+ * Memory is owned by the caller. The caller must make sure the buffer is large enough and at least as large as SizeInBytes.
+ */
+struct NNECORE_API FTensorBindingGPU
+{
+	void*	Data;
+	uint64	SizeInBytes;
+};
+
+/**
+ * The interface of a model that can run on GPU.
+ *
+ * Use UE::NNECore::GetRuntime<INNERuntimeGPU>(RuntimeName) to get a runtime capable of creating GPU models.
+ */
+class NNECORE_API IModelGPU
+{
+public:
+
+	virtual ~IModelGPU() = default;
+
+	/**
+	 * Get the input tensor descriptions as defined by the model, potentially with variable dimensions.
+	 *
+	 * @return An array containing a tensor descriptor for each input tensor of the model.
+	 */
+	virtual TConstArrayView<FTensorDesc> GetInputTensorDescs() const = 0;
+	
+	/**
+	 * Get the output tensor descriptions as defined by the model, potentially with variable dimensions.
+	 *
+	 * @return An array containing a tensor descriptor for each output tensor of the model.
+	 */
+	virtual TConstArrayView<FTensorDesc> GetOutputTensorDescs() const = 0;
+
+	/**
+	 * Get the input shapes.
+	 *
+	 * SetInputTensorShapes must be called prior of running a model.
+	 *
+	 * @return An array of input shapes or an empty array if SetInputTensorShapes has not been called.
+	 */
+	virtual TConstArrayView<FTensorShape> GetInputTensorShapes() const = 0;
+
+	/**
+	 * Getters for outputs shapes if they were already resolved.
+	 *
+	 * Output shapes might be resolved after a call to SetInputTensorShapes if the model and runtime supports it.
+	 * Otherwise they will be resolved while running the model
+	 *
+	 * @return An array of output shapes or an empty array if not resolved yet.
+	 */
+	virtual TConstArrayView<FTensorShape> GetOutputTensorShapes() const = 0;
+	
+	/**
+	 * Prepare the model to be run with the given input shape.
+	 *
+	 * The call is mandatory before a model can be run.
+	 * The function will run shape inference and resolve, if possible, the output shapes which can then be accessed by calling GetOutputTensorShapes().
+	 * This is a potentially expensive call and should be called lazily if possible.
+	 *
+	 * @param InInputShapes The input shapes to prepare the model with.
+	 * @return 0 on success or a non-zero number otherwise.
+	 */
+	virtual int SetInputTensorShapes(TConstArrayView<FTensorShape> InInputShapes) = 0;
+
+	/**
+	 * Evaluate the model synchronously.
+	 *
+	 * SetInputTensorShapes must be called prior to this call.
+	 * This function will block the calling thread until the inference is complete.
+	 * The caller owns the memory inside the bindings and must make sure that they are big enough.
+	 * Clients can call this function from an async task but must make sure the memory remains valid throughout the evaluation.
+	 *
+	 * @param InInputTensors An array containing tensor bindings for each input tensor with caller owned memory containing the input data.
+	 * @param InOutputTensors An array containing tensor bindings for each output tensor with caller owned memory big enough to contain the results on success.
+	 * @return 0 on success or a non-zero number otherwise.
+	 */
+	virtual int RunSync(TConstArrayView<FTensorBindingGPU> InInputTensors, TConstArrayView<FTensorBindingGPU> InOutputTensors) = 0;
+};
+
+} // UE::NNECore
+
+UINTERFACE()
+class NNECORE_API UNNERuntimeGPU : public UInterface
+{
+	GENERATED_BODY()
+};
+
+/**
+ * The interface of a neural network runtime capable of creating GPU models.
+ *
+ * Call UE::NNECore::GetRuntime<INNERuntimeGPU>(RuntimeName) to get a runtime implementing this interface.
+ */
+class NNECORE_API INNERuntimeGPU
+{
+	GENERATED_BODY()
+	
+public:
+
+	/**
+	 * Check if the runtime is able to create a model given some ModelData.
+	 *
+	 * @param ModelData The model data for which to create a model.
+	 * @return True if the runtime is able to create the model, false otherwise.
+	 */
+	virtual bool CanCreateModelGPU(TObjectPtr<UNNEModelData> ModelData) const = 0;
+	
+	/**
+	 * Create a model given some ModelData.
+	 *
+	 * The caller must make sure ModelData remains valid throughout the call.
+	 * ModelData is not required anymore after the model has been created.
+	 * The caller can decide to convert the result into a shared pointer if required (e.g. if the model needs to be shared with an async task for evaluation).
+	 *
+	 * @param ModelData The model data for which to create a model.
+	 * @return A caller owned model representing the neural network created from ModelData.
+	 */
+	virtual TUniquePtr<UE::NNECore::IModelGPU> CreateModelGPU(TObjectPtr<UNNEModelData> ModelData) = 0;
+};
\ No newline at end of file
diff --git a/Engine/Plugins/Experimental/NNE/Source/NNEQA/Private/NNEQAModel.cpp b/Engine/Plugins/Experimental/NNE/Source/NNEQA/Private/NNEQAModel.cpp
--- a/Engine/Plugins/Experimental/NNE/Source/NNEQA/Private/NNEQAModel.cpp
+++ b/Engine/Plugins/Experimental/NNE/Source/NNEQA/Private/NNEQAModel.cpp
@@ -25,6 +25,17 @@ namespace FModelQAHelpers
 		return Desc;
 	}
 
+	void ConvertBinding(TConstArrayView<NNECore::FTensorBindingCPU> InBindingsCPU, TArray<NNECore::FTensorBindingGPU>& OutBindingsGPU)
+	{
+		for (int32 Idx = 0; Idx < InBindingsCPU.Num(); ++Idx)
+		{
+			const NNECore::FTensorBindingCPU& BindingCPU = InBindingsCPU[Idx];
+			NNECore::FTensorBindingGPU& BindingGPU = OutBindingsGPU.Emplace_GetRef();
+			BindingGPU.Data = BindingCPU.Data;
+			BindingGPU.SizeInBytes = BindingCPU.SizeInBytes;
+		}
+	}
+
 	void ConvertBinding(FRDGBuilder& GraphBuilder,
 		TConstArrayView<NNECore::FTensorDesc> InTensorDescs,
 		TConstArrayView<NNECore::FTensorBindingCPU> InBindingsCPU,
@@ -112,6 +123,8 @@ TConstArrayView<NNECore::FTensorDesc> FModelQA::GetInputTensorDescs() const
 { 
 	if (ModelCPU.IsValid()) 
 		return ModelCPU->GetInputTensorDescs();
+	if (ModelGPU.IsValid())
+		return ModelGPU->GetInputTensorDescs();
 	if (ModelRDG.IsValid())
 		return ModelRDG->GetInputTensorDescs();
 	else
@@ -122,6 +135,8 @@ TConstArrayView<NNECore::FTensorDesc> FModelQA::GetOutputTensorDescs() const
 {
 	if (ModelCPU.IsValid())
 		return ModelCPU->GetOutputTensorDescs();
+	if (ModelGPU.IsValid())
+		return ModelGPU->GetOutputTensorDescs();
 	if (ModelRDG.IsValid())
 		return ModelRDG->GetOutputTensorDescs();
 	else
@@ -132,6 +147,8 @@ TConstArrayView<NNECore::FTensorShape> FModelQA::GetInputTensorShapes() const
 {
 	if (ModelCPU.IsValid())
 		return ModelCPU->GetInputTensorShapes();
+	if (ModelGPU.IsValid())
+		return ModelGPU->GetInputTensorShapes();
 	if (ModelRDG.IsValid())
 		return ModelRDG->GetInputTensorShapes();
 	else
@@ -142,6 +159,8 @@ TConstArrayView<NNECore::FTensorShape> FModelQA::GetOutputTensorShapes() const
 {
 	if (ModelCPU.IsValid())
 		return ModelCPU->GetOutputTensorShapes();
+	if (ModelGPU.IsValid())
+		return ModelGPU->GetOutputTensorShapes();
 	if (ModelRDG.IsValid())
 		return ModelRDG->GetOutputTensorShapes();
 	else
@@ -152,6 +171,8 @@ int FModelQA::SetInputTensorShapes(TConstArrayView<NNECore::FTensorShape> InInpu
 {
 	if (ModelCPU.IsValid())
 		return ModelCPU->SetInputTensorShapes(InInputShapes);
+	if (ModelGPU.IsValid())
+		return ModelGPU->SetInputTensorShapes(InInputShapes);
 	if (ModelRDG.IsValid())
 		return ModelRDG->SetInputTensorShapes(InInputShapes);
 	else
@@ -190,6 +211,14 @@ int FModelQA::RunSync(TConstArrayView<NNECore::FTensorBindingCPU> InInputBinding
 	{
 		return ModelCPU->RunSync(InInputBindings, InOutputBindings);
 	}
+	if (ModelGPU.IsValid())
+	{
+		TArray<NNECore::FTensorBindingGPU> InputBindingsGPU;
+		TArray<NNECore::FTensorBindingGPU> OutputBindingsGPU;
+		FModelQAHelpers::ConvertBinding(InInputBindings, InputBindingsGPU);
+		FModelQAHelpers::ConvertBinding(InOutputBindings, OutputBindingsGPU);
+		return ModelGPU->RunSync(InputBindingsGPU, OutputBindingsGPU);
+	}
 	else if (ModelRDG.IsValid())
 	{
 		//Note: ok if pending?
@@ -251,6 +280,7 @@ TUniquePtr<FModelQA> FModelQA::MakeModelQA(const FNNEModelRaw& ONNXModelData, co
 	TUniquePtr<FModelQA> ModelQA = MakeUnique<FModelQA>();
 	TWeakInterfacePtr<INNERuntime> Runtime = NNECore::GetRuntime<INNERuntime>(RuntimeName);
 	TWeakInterfacePtr<INNERuntimeCPU> RuntimeCPU = NNECore::GetRuntime<INNERuntimeCPU>(RuntimeName);
+	TWeakInterfacePtr<INNERuntimeGPU> RuntimeGPU = NNECore::GetRuntime<INNERuntimeGPU>(RuntimeName);
 	TWeakInterfacePtr<INNERuntimeRDG> RuntimeRDG = NNECore::GetRuntime<INNERuntimeRDG>(RuntimeName);
 	TObjectPtr<UNNEModelData> ModelData = NewObject<UNNEModelData>();
 			
@@ -264,6 +294,10 @@ TUniquePtr<FModelQA> FModelQA::MakeModelQA(const FNNEModelRaw& ONNXModelData, co
 	{
 		ModelQA->ModelCPU = RuntimeCPU->CreateModelCPU(ModelData);
 	}
+	else if (RuntimeGPU.IsValid())
+	{
+		ModelQA->ModelGPU = RuntimeGPU->CreateModelGPU(ModelData);
+	}
 	else if (RuntimeRDG.IsValid())
 	{
 		ModelQA->ModelRDG = RuntimeRDG->CreateModelRDG(ModelData);
diff --git a/Engine/Plugins/Experimental/NNE/Source/NNEQA/Private/NNEQAModel.h b/Engine/Plugins/Experimental/NNE/Source/NNEQA/Private/NNEQAModel.h
--- a/Engine/Plugins/Experimental/NNE/Source/NNEQA/Private/NNEQAModel.h
+++ b/Engine/Plugins/Experimental/NNE/Source/NNEQA/Private/NNEQAModel.h
@@ -6,6 +6,7 @@
 #include "NNECore.h"
 #include "NNECoreRuntimeCPU.h"
 #include "NNECoreRuntimeFormat.h"
+#include "NNECoreRuntimeGPU.h"
 #include "NNECoreRuntimeRDG.h"
 #include "NNECoreTypes.h"
 #include "RenderGraphBuilder.h"
@@ -23,6 +24,7 @@ class FModelQA : NNECore::IModelCPU
 private:
 
 	TUniquePtr<NNECore::IModelCPU> ModelCPU;
+	TUniquePtr<NNECore::IModelGPU> ModelGPU;
 	TUniquePtr<NNECore::IModelRDG> ModelRDG;
 
 public:
diff --git a/Engine/Plugins/Experimental/NNE/Source/NNERuntimeORT/Private/NNERuntimeORT.cpp b/Engine/Plugins/Experimental/NNE/Source/NNERuntimeORT/Private/NNERuntimeORT.cpp
--- a/Engine/Plugins/Experimental/NNE/Source/NNERuntimeORT/Private/NNERuntimeORT.cpp
+++ b/Engine/Plugins/Experimental/NNE/Source/NNERuntimeORT/Private/NNERuntimeORT.cpp
@@ -55,7 +55,7 @@ void UNNERuntimeORTDmlImpl::Init()
 }
 
 #if PLATFORM_WINDOWS
-bool UNNERuntimeORTDmlImpl::CanCreateModelCPU(TObjectPtr<UNNEModelData> ModelData) const
+bool UNNERuntimeORTDmlImpl::CanCreateModelGPU(TObjectPtr<UNNEModelData> ModelData) const
 {
 	check(ModelData != nullptr);
 
@@ -73,14 +73,14 @@ bool UNNERuntimeORTDmlImpl::CanCreateModelCPU(TObjectPtr<UNNEModelData> ModelDat
 	return bResult;
 }
 
-TUniquePtr<UE::NNECore::IModelCPU> UNNERuntimeORTDmlImpl::CreateModelCPU(TObjectPtr<UNNEModelData> ModelData)
+TUniquePtr<UE::NNECore::IModelGPU> UNNERuntimeORTDmlImpl::CreateModelGPU(TObjectPtr<UNNEModelData> ModelData)
 {
 	check(ModelData != nullptr);
 	check(ORTEnvironment.IsValid());
 
-	if (!CanCreateModelCPU(ModelData))
+	if (!CanCreateModelGPU(ModelData))
 	{
-		return TUniquePtr<UE::NNECore::IModelCPU>();
+		return TUniquePtr<UE::NNECore::IModelGPU>();
 	}
 
 	const UE::NNERuntimeORT::Private::FRuntimeConf InConf;
@@ -90,22 +90,22 @@ TUniquePtr<UE::NNECore::IModelCPU> UNNERuntimeORTDmlImpl::CreateModelCPU(TObject
 	if (!Model->Init(Data))
 	{
 		delete Model;
-		return TUniquePtr<UE::NNECore::IModelCPU>();
+		return TUniquePtr<UE::NNECore::IModelGPU>();
 	}
-	UE::NNECore::IModelCPU* IModel = static_cast<UE::NNECore::IModelCPU*>(Model);
-	return TUniquePtr<UE::NNECore::IModelCPU>(IModel);
+	UE::NNECore::IModelGPU* IModel = static_cast<UE::NNECore::IModelGPU*>(Model);
+	return TUniquePtr<UE::NNECore::IModelGPU>(IModel);
 }
 
 #else // PLATFORM_WINDOWS
 
-bool UNNERuntimeORTDmlImpl::CanCreateModelCPU(TObjectPtr<UNNEModelData> ModelData) const
+bool UNNERuntimeORTDmlImpl::CanCreateModelGPU(TObjectPtr<UNNEModelData> ModelData) const
 {
 	return false;
 }
 
-TUniquePtr<UE::NNECore::IModelCPU> UNNERuntimeORTDmlImpl::CreateModelCPU(TObjectPtr<UNNEModelData> ModelData)
+TUniquePtr<UE::NNECore::IModelGPU> UNNERuntimeORTDmlImpl::CreateModelGPU(TObjectPtr<UNNEModelData> ModelData)
 {
-	return TUniquePtr<UE::NNECore::IModelCPU>();
+	return TUniquePtr<UE::NNECore::IModelGPU>();
 }
 
 #endif // PLATFORM_WINDOWS
\ No newline at end of file
diff --git a/Engine/Plugins/Experimental/NNE/Source/NNERuntimeORT/Private/NNERuntimeORT.h b/Engine/Plugins/Experimental/NNE/Source/NNERuntimeORT/Private/NNERuntimeORT.h
--- a/Engine/Plugins/Experimental/NNE/Source/NNERuntimeORT/Private/NNERuntimeORT.h
+++ b/Engine/Plugins/Experimental/NNE/Source/NNERuntimeORT/Private/NNERuntimeORT.h
@@ -4,7 +4,7 @@
 
 #include "CoreMinimal.h"
 #include "NNECoreRuntime.h"
-#include "NNECoreRuntimeCPU.h"
+#include "NNECoreRuntimeGPU.h"
 #include "Templates/UniquePtr.h"
 #include "UObject/Class.h"
 #include "UObject/Object.h"
@@ -20,7 +20,7 @@ NNE_THIRD_PARTY_INCLUDES_END
 #include "NNERuntimeORT.generated.h"
 
 UCLASS()
-class UNNERuntimeORTDmlImpl : public UObject, public INNERuntime, public INNERuntimeCPU
+class UNNERuntimeORTDmlImpl : public UObject, public INNERuntime, public INNERuntimeGPU
 {
 	GENERATED_BODY()
 
@@ -40,6 +40,6 @@ public:
 	virtual bool CanCreateModelData(FString FileType, TConstArrayView<uint8> FileData) const override;
 	virtual TArray<uint8> CreateModelData(FString FileType, TConstArrayView<uint8> FileData) override;
 
-	virtual bool CanCreateModelCPU(TObjectPtr<UNNEModelData> ModelData) const override;
-	virtual TUniquePtr<UE::NNECore::IModelCPU> CreateModelCPU(TObjectPtr<UNNEModelData> ModelData) override;
+	virtual bool CanCreateModelGPU(TObjectPtr<UNNEModelData> ModelData) const override;
+	virtual TUniquePtr<UE::NNECore::IModelGPU> CreateModelGPU(TObjectPtr<UNNEModelData> ModelData) override;
 };
\ No newline at end of file
diff --git a/Engine/Plugins/Experimental/NNE/Source/NNERuntimeORT/Private/NNERuntimeORTModel.cpp b/Engine/Plugins/Experimental/NNE/Source/NNERuntimeORT/Private/NNERuntimeORTModel.cpp
--- a/Engine/Plugins/Experimental/NNE/Source/NNERuntimeORT/Private/NNERuntimeORTModel.cpp
+++ b/Engine/Plugins/Experimental/NNE/Source/NNERuntimeORT/Private/NNERuntimeORTModel.cpp
@@ -191,7 +191,7 @@ namespace UE::NNERuntimeORT::Private
 		OutputTensorShapes.Reset();
 
 		// Verify input shape are valid for the model and set InputTensorShapes
-		if (FModelBase<IModelCPU>::SetInputTensorShapes(InInputShapes) != 0)
+		if (FModelBase<IModelGPU>::SetInputTensorShapes(InInputShapes) != 0)
 		{
 			return -1;
 		}
@@ -225,7 +225,7 @@ namespace UE::NNERuntimeORT::Private
 		return 0;
 	}
 
-	int FModelORT::RunSync(TConstArrayView<NNECore::FTensorBindingCPU> InInputBindings, TConstArrayView<NNECore::FTensorBindingCPU> InOutputBindings)
+	int FModelORT::RunSync(TConstArrayView<NNECore::FTensorBindingGPU> InInputBindings, TConstArrayView<NNECore::FTensorBindingGPU> InOutputBindings)
 	{
 		DECLARE_SCOPE_CYCLE_COUNTER(TEXT("FModelORT_Run"), STAT_FModelORT_Run, STATGROUP_NNE);
 
@@ -322,18 +322,6 @@ namespace UE::NNERuntimeORT::Private
 		InputTransferStatisticsEstimator.ResetStats();
 	}
 
-	bool FModelORTCpu::InitializedAndConfigureMembers()
-	{
-		if (!FModelORT::InitializedAndConfigureMembers())
-		{
-			return false;
-		}
-
-		SessionOptions->EnableCpuMemArena();
-
-		return true;
-	}
-
 #if PLATFORM_WINDOWS
 	bool FModelORTDml::InitializedAndConfigureMembers()
 	{
diff --git a/Engine/Plugins/Experimental/NNE/Source/NNERuntimeORT/Private/NNERuntimeORTModel.h b/Engine/Plugins/Experimental/NNE/Source/NNERuntimeORT/Private/NNERuntimeORTModel.h
--- a/Engine/Plugins/Experimental/NNE/Source/NNERuntimeORT/Private/NNERuntimeORTModel.h
+++ b/Engine/Plugins/Experimental/NNE/Source/NNERuntimeORT/Private/NNERuntimeORTModel.h
@@ -4,7 +4,7 @@
 
 #include "CoreMinimal.h"
 #include "NNECoreModelBase.h"
-#include "NNECoreRuntimeCPU.h"
+#include "NNECoreRuntimeGPU.h"
 #include "NNECoreTensor.h"
 #include "NNECoreTypes.h"
 #include "NNEProfilingStatistics.h"
@@ -25,7 +25,7 @@ namespace UE::NNERuntimeORT::Private
 		EThreadPriority ThreadPriority = EThreadPriority::TPri_Normal;
 	};
 
-	class FModelORT : public NNECore::Internal::FModelBase<NNECore::IModelCPU>
+	class FModelORT : public NNECore::Internal::FModelBase<NNECore::IModelGPU>
 	{
 
 	public:
@@ -34,7 +34,7 @@ namespace UE::NNERuntimeORT::Private
 		virtual ~FModelORT() = default;
 
 		virtual int SetInputTensorShapes(TConstArrayView<NNECore::FTensorShape> InInputShapes) override;
-		virtual int RunSync(TConstArrayView<NNECore::FTensorBindingCPU> InInputBindings, TConstArrayView<NNECore::FTensorBindingCPU> InOutputBindings) override;
+		virtual int RunSync(TConstArrayView<NNECore::FTensorBindingGPU> InInputBindings, TConstArrayView<NNECore::FTensorBindingGPU> InOutputBindings) override;
 
 		bool Init(TConstArrayView<uint8> ModelData);
 		bool IsLoaded() const;
@@ -78,16 +78,6 @@ namespace UE::NNERuntimeORT::Private
 
 	};
 
-	class FModelORTCpu : public FModelORT
-	{
-	public:
-		FModelORTCpu() {}
-		FModelORTCpu(Ort::Env* InORTEnvironment, const FRuntimeConf& InRuntimeConf) : FModelORT(InORTEnvironment, InRuntimeConf) {}
-		virtual ~FModelORTCpu() = default;
-	private:
-		virtual bool InitializedAndConfigureMembers();
-	};
-
 #if PLATFORM_WINDOWS
 	class FModelORTDml : public FModelORT
 	{
diff --git a/Engine/Plugins/Experimental/NNE/Source/NNERuntimeORT/Private/NNERuntimeORTUtils.h b/Engine/Plugins/Experimental/NNE/Source/NNERuntimeORT/Private/NNERuntimeORTUtils.h
--- a/Engine/Plugins/Experimental/NNE/Source/NNERuntimeORT/Private/NNERuntimeORTUtils.h
+++ b/Engine/Plugins/Experimental/NNE/Source/NNERuntimeORT/Private/NNERuntimeORTUtils.h
@@ -4,7 +4,7 @@
 
 #include "HAL/UnrealMemory.h"
 #include "NNECore.h"
-#include "NNECoreRuntimeCPU.h"
+#include "NNECoreRuntimeGPU.h"
 #include "NNECoreTypes.h"
 #include "NNECoreTensor.h"
 #include "NNEThirdPartyWarningDisabler.h"
@@ -118,7 +118,7 @@ namespace UE::NNERuntimeORT::Private
 	}
 
 	inline void BindTensorsToORT(
-		TConstArrayView<NNECore::FTensorBindingCPU> InBindingTensors,
+		TConstArrayView<NNECore::FTensorBindingGPU> InBindingTensors,
 		TConstArrayView<NNECore::Internal::FTensor> InTensors,
 		TConstArrayView<ONNXTensorElementDataType> InTensorsORTType,
 		const Ort::MemoryInfo* InAllocatorInfo,
@@ -136,7 +136,7 @@ namespace UE::NNERuntimeORT::Private
 
 		for (uint32 Index = 0; Index < NumBinding; ++Index)
 		{
-			const NNECore::FTensorBindingCPU& Binding = InBindingTensors[Index];
+			const NNECore::FTensorBindingGPU& Binding = InBindingTensors[Index];
 			const NNECore::Internal::FTensor& Tensor = InTensors[Index];
 			const ONNXTensorElementDataType CurrentORTType = InTensorsORTType[Index];
 
@@ -164,7 +164,7 @@ namespace UE::NNERuntimeORT::Private
 
 	inline void CopyFromORTToBindings(
 		TConstArrayView<Ort::Value> InOrtTensors,
-		TConstArrayView<NNECore::FTensorBindingCPU> InBindingTensors,
+		TConstArrayView<NNECore::FTensorBindingGPU> InBindingTensors,
 		TConstArrayView<NNECore::FTensorDesc> InTensorDescs,
 		TArray<NNECore::Internal::FTensor>& OutTensors
 	)
@@ -180,7 +180,7 @@ namespace UE::NNERuntimeORT::Private
 
 		for (uint32 Index = 0; Index < NumBinding; ++Index)
 		{
-			const NNECore::FTensorBindingCPU& Binding = InBindingTensors[Index];
+			const NNECore::FTensorBindingGPU& Binding = InBindingTensors[Index];
 			const NNECore::FTensorDesc& TensorDesc = InTensorDescs[Index];
 			const Ort::Value& OrtTensor = InOrtTensors[Index];
 			const std::vector<int64_t>& OrtShape = OrtTensor.GetTensorTypeAndShapeInfo().GetShape();
