Add RDG DirectML Cast operator

#rnx

diff --git a/Engine/Plugins/Experimental/NNE/Source/NNEQA/Resources/NNEQATestDesc.json b/Engine/Plugins/Experimental/NNE/Source/NNEQA/Resources/NNEQATestDesc.json
--- a/Engine/Plugins/Experimental/NNE/Source/NNEQA/Resources/NNEQATestDesc.json
+++ b/Engine/Plugins/Experimental/NNE/Source/NNEQA/Resources/NNEQATestDesc.json
@@ -227,8 +227,7 @@
 					]},
 				{	"target": "Atanh" },
 				{	"target": "BitShift", "skip": true, "comment": "TODO need attributes" },
-				{	"target": "Cast",     "skip": true, "comment": "TODO need attributes" },
-				{	"target": "Ceil"},
+				{	"target": "Ceil" },
 				{	"target": "Clip", "additional_datasets": [ "Elementwise.Clip.AttributeTest" ] },
 				{	"target": "Cos" },
 				{	"target": "Cosh" },
@@ -362,6 +361,15 @@
 				{ "target": "Gemm", "tags": [ "SmokeTest" ] }
 			]
 		},
+		{
+			"category": "CastOp",
+			"runtimes": [	
+				{ "name": "NNERuntimeRDGHlsl", "skip": true }
+			],
+			"targets": [
+				{ "target": "Cast", "tags": [ "SmokeTest" ] }
+			]
+		},
 		{
 			"category": "Convolution",
 			"targets": [
@@ -564,6 +572,12 @@
 				{ "inputs": [ { "shape": [ 1,3,720,1280] } ] }
 			]
 		},
+		{
+			"name": "CastOp",
+			"datasets": [
+				{ "inputs": [ { "shape": [ 32 ] } ] , "outputs": [ { "shape": [ 32 ], "type": "float" } ]}
+			]
+		},
 		{
 			"name": "Elementwise.Unary.Regular",
 			"datasets": [
@@ -1382,6 +1396,12 @@
 				{ "attributes": [ { "name": "transA", "value": 1, "type": "Int32" }, { "name": "transB", "value": 1, "type": "Int32" } ] }
 			]
 		},
+		{
+			"name": "CastOp",
+			"attribute_maps": [
+				{ "attributes": [ { "name": "to", "value": 1, "type": "Int32" } ] }
+			]
+		},
 		{
 			"name": "Convolution.DefaultPadding",
 			"attribute_maps": [
diff --git a/Engine/Plugins/Experimental/NNE/Source/NNERuntimeRDG/Private/Dml/NNEDmlOperatorCast.cpp b/Engine/Plugins/Experimental/NNE/Source/NNERuntimeRDG/Private/Dml/NNEDmlOperatorCast.cpp
new file mode 100644
--- /dev/null
+++ b/Engine/Plugins/Experimental/NNE/Source/NNERuntimeRDG/Private/Dml/NNEDmlOperatorCast.cpp
@@ -0,0 +1,87 @@
+// Copyright Epic Games, Inc. All Rights Reserved.
+
+#ifdef NNE_USE_DIRECTML
+#include "NNEDmlOperator.h"
+#include "NNEDmlOperatorUtils.h"
+#include "NNEUtilsLogHelper.h"
+
+namespace UE::NNERuntimeRDG::Private::Dml
+{
+
+class FOperatorDmlCast : public FOperatorDml
+{
+public:
+
+	static FOperatorDml* Create()
+	{
+		return new FOperatorDmlCast();
+	}
+
+	//
+	//
+	//
+	virtual bool Initialize(IDMLDevice* Device, TArrayView<const NNECore::Internal::FTensor> InputTensors, TArrayView<const NNECore::Internal::FTensor> OutputTensors, const NNECore::FAttributeMap& Attributes) override
+	{
+		check(InputTensors.Num() == 1);
+		check(OutputTensors.Num() == 1);
+		
+		const NNECore::Internal::FTensor& InputTensor = InputTensors[0];
+		const NNECore::Internal::FTensor& OutputTensor = OutputTensors[0];
+
+		ENNETensorDataType To = (ENNETensorDataType) Attributes.GetValue<int32>(TEXT("to"));
+		
+		if (To != OutputTensor.GetDataType())
+		{
+			UE_LOG(LogNNE, Warning, TEXT("Cast should output a tensor of type %d but was of type %d."), To, OutputTensor.GetDataType());
+			return false;
+		}
+
+		TConstArrayView<uint32> InputShape = InputTensor.GetShape().GetData();
+		TConstArrayView<uint32> OutputShape = OutputTensor.GetShape().GetData();
+
+		if (InputShape.Num() != OutputShape.Num())
+		{
+			UE_LOG(LogNNE, Warning, TEXT("Cast input and output shapes need to have a same rank"));
+			return false;
+		}
+
+		for (int32 Idx = 0; Idx < InputShape.Num(); ++Idx)
+		{
+			if (InputShape[Idx] != OutputShape[Idx])
+			{
+				UE_LOG(LogNNE, Warning, TEXT("Input shape and output shape need to have a same dimension at dim %d (%d != %d"), Idx, InputShape[Idx], OutputShape[Idx]);
+				return false;
+			}
+		}
+
+		// Initialize tensor descriptors
+		DmlUtil::FTensorDesc	DmlInputTensor{};
+		DmlUtil::FTensorDesc	DmlOutputTensor{};
+
+		if (!DmlInputTensor.InitFromTensor(InputTensor, InputTensor.GetShape().Rank()))
+		{
+			UE_LOG(LogNNE, Warning, TEXT("Failed to initialize tensor(s) for DML inference"));
+			return false;
+		}
+
+		if (!DmlOutputTensor.InitFromTensor(OutputTensor, InputTensor.GetShape().Rank()))
+		{
+			UE_LOG(LogNNE, Warning, TEXT("Failed to initialize tensor(s) for DML inference"));
+			return false;
+		}
+
+		DML_CAST_OPERATOR_DESC	OpDesc{};
+
+		OpDesc.InputTensor = &DmlInputTensor.Desc;
+		OpDesc.OutputTensor = &DmlOutputTensor.Desc;
+
+		return CreateOperator(Device, DML_OPERATOR_DESC { DML_OPERATOR_CAST, &OpDesc });
+	}
+};
+
+// Register operator on Module startup
+NNE_DML_REGISTER_OP(Cast)
+
+} // namespace UE::NNERuntimeRDG::Private::Dml
+
+#endif // NNE_USE_DIRECTML
