NNE - Extend support of Resize op implementation in RDG Dml

#rnx

diff --git a/Engine/Plugins/Experimental/NNE/Source/NNEQA/Resources/NNEQATestDesc.json b/Engine/Plugins/Experimental/NNE/Source/NNEQA/Resources/NNEQATestDesc.json
--- a/Engine/Plugins/Experimental/NNE/Source/NNEQA/Resources/NNEQATestDesc.json
+++ b/Engine/Plugins/Experimental/NNE/Source/NNEQA/Resources/NNEQATestDesc.json
@@ -431,6 +431,16 @@
 				{ "target": "Upsample" }
 			]
 		},
+		{
+			"category": "Resize",
+			"runtimes": [	
+				{ "name": "NNERuntimeRDGHlsl", "skip": true },
+				{ "name": "NNERuntimeORTDml", "skip": true }
+			],
+			"targets": [
+				{ "target": "Resize", "absolute_tolerance": 1e-5 }
+			]
+		},
 		{
 			"category": "Gather",
 			"runtimes": [	
@@ -948,6 +958,13 @@
 				{ "inputs":  [ { "shape": [ 2, 1, 3, 1, 2 ] } ],  "weights": [ { "shape": [ 5 ], "source": [ "1", "2", "2", "1", "2" ], "type": "float" } ], "outputs": [ { "shape": [ 2, 2, 6, 1, 4 ]} ] }
 			]
 		},
+		{
+			"name": "Resize.Downsample.Linear",
+			"datasets": [
+				{ "inputs":  [ { "shape": [ 10, 10 ] } ],	"weights": [ { "shape": [ 4 ], "source": [ "2", "1", "8", "9" ], "type": "float" }, { "shape": [ 2 ], "source": [ "0.6", "0.8"  ], "type": "float" } ],	"outputs": [ { "shape": [ 6, 8 ]} ] },
+				{ "inputs":  [ { "shape": [ 20, 15 ] } ],	"weights": [ { "shape": [ 4 ], "source": [ "5", "3", "19", "14" ], "type": "float" }, { "shape": [ 2 ], "source": [ "0.5", "0.4"  ], "type": "float" } ],	"outputs": [ { "shape": [ 10, 6 ]} ] }
+			]
+		},
 		{
 			"name": "Identity",
 			"datasets": [
@@ -1641,6 +1658,25 @@
 				] }
 			]
 		},
+		{
+			"name": "Resize.Downsample.Linear",
+			"attribute_maps": [
+				{
+					"attributes": [ 
+						{ "name": "mode", "value": "linear", "type": "String" }
+					]
+				}
+			]
+		},
+		{
+			"name": "Resize.Downsample.Linear.Transform",
+			"multiply_with_attribute_sets": [ "Resize.Downsample.Linear" ],
+			"attribute_maps": [
+				{ "attributes": [ { "name": "coordinate_transformation_mode", "value": "align_corners", "type": "String" } ] },
+				{ "attributes": [ { "name": "coordinate_transformation_mode", "value": "asymmetric", "type": "String" } ] },
+				{ "attributes": [ { "name": "coordinate_transformation_mode", "value": "tf_half_pixel_for_nn", "type": "String" } ] }
+			]
+		},
 		{
 			"name": "GRU.Defaults",
 			"attribute_maps": [
diff --git a/Engine/Plugins/Experimental/NNE/Source/NNERuntimeRDG/Private/Dml/NNEDmlOperatorResample.cpp b/Engine/Plugins/Experimental/NNE/Source/NNERuntimeRDG/Private/Dml/NNEDmlOperatorResample.cpp
--- a/Engine/Plugins/Experimental/NNE/Source/NNERuntimeRDG/Private/Dml/NNEDmlOperatorResample.cpp
+++ b/Engine/Plugins/Experimental/NNE/Source/NNERuntimeRDG/Private/Dml/NNEDmlOperatorResample.cpp
@@ -76,15 +76,28 @@ public:
 
 		if (!ScaleTensor.HasPreparedData())
 		{
-			UE_LOG(LogNNE, Warning, TEXT("scale should be a constant tensor, it is here a variable tensor of name %s."), *ScaleTensor.GetName());
+			UE_LOG(LogNNE, Warning, TEXT("scales should be a constant tensor, it is here a variable tensor of name %s."), *ScaleTensor.GetName());
+			return false;
+		}
+
+		if (ScaleTensor.GetShape().Volume() != InputTensor.GetShape().Rank())
+		{
+			UE_LOG(LogNNE, Warning, TEXT("scales tensor should contain N entries, where N is rank of X."));
 			return false;
 		}
 		
 		// Read attributes
 		DML_INTERPOLATION_MODE Mode = ModeFromString(Attributes.GetValue<FString>(TEXT("mode")));
 
+		DmlUtil::FSmallArray<float> InputPixelOffsets, OutputPixelOffsets;
+		InputPixelOffsets.Init(0.5f, InputTensor.GetShape().Rank());
+		OutputPixelOffsets.Init(-0.5f, InputTensor.GetShape().Rank());
+
+		DmlUtil::FSmallArray<float> Scales ( ScaleTensor.GetPreparedData<float>() );
+
 		if constexpr (IsResize)
 		{
+
 			if (Mode == DML_INTERPOLATION_MODE_NEAREST_NEIGHBOR)
 			{
 				FString NeareastMode = Attributes.GetValueOrDefault<FString>(TEXT("nearest_mode"), TEXT("round_prefer_floor"));
@@ -95,10 +108,71 @@ public:
 			}
 
 			FString CoordinateTransformationMode = Attributes.GetValueOrDefault<FString>(TEXT("coordinate_transformation_mode"), TEXT("half_pixel"));
-			if (FCString::Stricmp(*CoordinateTransformationMode, TEXT("half_pixel")) != 0)
+			for(int Idx = 0; Idx < InputTensor.GetShape().Rank(); ++Idx)
 			{
-				UE_LOG(LogNNE, Warning, TEXT("Unsupported coordinate transformation mode:%s, using half_pixel instead"), *CoordinateTransformationMode);
+				float LengthResized = (float) OutputTensor.GetShape().GetData()[Idx];
+				float LengthOriginal = (float) InputTensor.GetShape().GetData()[Idx];
+				
+				if(CoordinateTransformationMode == TEXT("align_corners"))
+				{
+					Scales[Idx] = (LengthResized - 1.0f) / (LengthOriginal - 1.0f);
+					InputPixelOffsets[Idx] = 0.f;
+					OutputPixelOffsets[Idx] = 0.f;
+				}
+				else if(CoordinateTransformationMode == TEXT("asymmetric"))
+				{
+					InputPixelOffsets[Idx] = 0.f;
+					OutputPixelOffsets[Idx] = 0.f;
+				}
+				else if(CoordinateTransformationMode == TEXT("tf_half_pixel_for_nn"))
+				{
+					InputPixelOffsets[Idx] = 0.0f;
+					OutputPixelOffsets[Idx] = -0.5f;
+				}
+				else if(CoordinateTransformationMode == TEXT("tf_crop_and_resize"))
+				{
+					// NOTE: no tests for this, ORT erroneously puts all 0.0fs in the output tensor in this case.
+					
+					const NNECore::Internal::FTensor& RoiTensor = InputTensors[1];
+					if (!RoiTensor.HasPreparedData())
+					{
+						UE_LOG(LogNNE, Warning, TEXT("roi should be a constant tensor, it is here a variable tensor of name %s."), *RoiTensor.GetName());
+						return false;
+					}
+					if (RoiTensor.GetShape().Rank() != 1)
+					{
+						UE_LOG(LogNNE, Warning, TEXT("roi tensor should be 1-D."));
+						return false;
+					}
+					if (RoiTensor.GetShape().GetData()[0] != 2 * InputTensor.GetShape().Rank())
+					{
+						UE_LOG(LogNNE, Warning, TEXT("roi tensor should contain 2*N entries, where N is rank of X."));
+						return false;
+					}
+					float Start = RoiTensor.GetPreparedData<float>()[Idx];
+					float End = RoiTensor.GetPreparedData<float>()[Idx + InputTensor.GetShape().Rank()];
+					if(LengthResized > 1)
+					{
+						Scales[Idx] = (LengthResized - 1.0f) / FMath::Max( (End - Start) * (LengthOriginal - 1.0f), 1.0f );
+						InputPixelOffsets[Idx] = Start * (1.0f - LengthOriginal);
+						OutputPixelOffsets[Idx] = 0.0f;
+					}
+					else
+					{
+						UE_LOG(LogNNE, Warning, TEXT("Unsupported combination of transformation mode tf_crop_and_resize and length of resized dimension (%d) <= 1"), Idx);
+						return false;
+					}
+					
+				}
+				else
+				{
+					if(CoordinateTransformationMode != TEXT("half_pixel"))
+					{
+						UE_LOG(LogNNE, Warning, TEXT("Unsupported coordinate transformation mode:%s, using half_pixel instead"), *CoordinateTransformationMode);
+					}
+				}
 			}
+			
 		}
 
 		// Initialize tensor descriptors
@@ -123,7 +197,7 @@ public:
 			ConstantCPUInputs.Add(i);
 		}
 
-		TConstArrayView<float> Scales = ScaleTensor.GetPreparedData<float>();
+		
 		
 		// Find any useless dimensions of size 1 that occur in both input and output
 		DmlUtil::FSmallUIntArray	SqueezeInds;
@@ -147,21 +221,25 @@ public:
 			RemoveValuesByIndex(SqueezeInds, SqueezedInputShape, true);
 			RemoveValuesByIndex(SqueezeInds, SqueezedOutputShape, true);
 			RemoveValuesByIndex(SqueezeInds, ScaleValues, true);
+			RemoveValuesByIndex(SqueezeInds, InputPixelOffsets, true);
+			RemoveValuesByIndex(SqueezeInds, OutputPixelOffsets, true);
 
 			DmlInputTensor.UpdateShapeAndStrides(SqueezedInputShape);
 			DmlOutputTensor.UpdateShapeAndStrides(SqueezedOutputShape);
 			Scales = ScaleValues;
 		}
 
-		DML_RESAMPLE_OPERATOR_DESC	OpDesc{};
+		DML_RESAMPLE1_OPERATOR_DESC	OpDesc{};
 
 		OpDesc.InputTensor = &DmlInputTensor.Desc;
 		OpDesc.OutputTensor = &DmlOutputTensor.Desc;
 		OpDesc.InterpolationMode = Mode;
-		OpDesc.ScaleCount = Scales.Num();
+		OpDesc.DimensionCount = Scales.Num();
 		OpDesc.Scales = Scales.GetData();
+		OpDesc.InputPixelOffsets = InputPixelOffsets.GetData();
+		OpDesc.OutputPixelOffsets = OutputPixelOffsets.GetData();
 
-		return CreateOperator(Device, DML_OPERATOR_DESC { DML_OPERATOR_RESAMPLE, &OpDesc });
+		return CreateOperator(Device, DML_OPERATOR_DESC { DML_OPERATOR_RESAMPLE1, &OpDesc });
 	}
 };
 
diff --git a/Engine/Plugins/Experimental/NNE/Source/NNEUtils/Private/NNEUtilsModelBuilderONNX.cpp b/Engine/Plugins/Experimental/NNE/Source/NNEUtils/Private/NNEUtilsModelBuilderONNX.cpp
--- a/Engine/Plugins/Experimental/NNE/Source/NNEUtils/Private/NNEUtilsModelBuilderONNX.cpp
+++ b/Engine/Plugins/Experimental/NNE/Source/NNEUtils/Private/NNEUtilsModelBuilderONNX.cpp
@@ -375,6 +375,7 @@ NNEUTILS_API bool CreateONNXModelForOperator(bool UseVariadicShapeForModel, cons
 		OperatorName == TEXT("ReduceProd") ||			// current implementation is opset 11 (next version is 13)
 		OperatorName == TEXT("ReduceSum") ||			// current implementation is opset 11 (next version is 13)
 		OperatorName == TEXT("ReduceSumSquare") ||		// current implementation is opset 11 (next version is 13)
+		OperatorName == TEXT("Resize") ||				// current implementation is opset 11 (next version is 13)
 		OperatorName == TEXT("Squeeze") ||				// current implementation is opset 11 (next version is 13)
 		OperatorName == TEXT("Unsqueeze")				// current implementation is opset 11 (next version is 13)
 		)
