Fixed shadowing warnings.

diff --git a/Engine/Plugins/Experimental/NNE/Source/NNERuntimeRDG/Private/Dml/NNEDmlModel.cpp b/Engine/Plugins/Experimental/NNE/Source/NNERuntimeRDG/Private/Dml/NNEDmlModel.cpp
--- a/Engine/Plugins/Experimental/NNE/Source/NNERuntimeRDG/Private/Dml/NNEDmlModel.cpp
+++ b/Engine/Plugins/Experimental/NNE/Source/NNERuntimeRDG/Private/Dml/NNEDmlModel.cpp
@@ -123,12 +123,12 @@ public:
 	}
 
 #ifdef NNE_USE_D3D12_RESOURCES
-	void Bind(IDMLOperatorInitializer* OpInit, TConstArrayView<FRHIBuffer*> InputBuffers, ID3D12Resource* PersistBuff, ID3D12Resource* TempBuff = nullptr)
+	void Bind(IDMLOperatorInitializer* InOpInit, TConstArrayView<FRHIBuffer*> InputBuffers, ID3D12Resource* InPersistBuff, ID3D12Resource* InTempBuff = nullptr)
 #else
-	void Bind(IDMLOperatorInitializer* OpInit, TConstArrayView<FRHIBuffer*> InputBuffers, FRHIBuffer* PersistBuff, FRHIBuffer* TempBuff = nullptr)
+	void Bind(IDMLOperatorInitializer* InOpInit, TConstArrayView<FRHIBuffer*> InputBuffers, FRHIBuffer* InPersistBuff, FRHIBuffer* InTempBuff = nullptr)
 #endif
 	{
-		Reset(OpInit);
+		Reset(InOpInit);
 
 		TArray<DML_BUFFER_BINDING, TInlineAllocator<MaxNumInputs>> Inputs;
 
@@ -152,12 +152,12 @@ public:
 		DML_BUFFER_BINDING			PersistBind {};
 		DML_BINDING_DESC			PersistBindDesc{ DML_BINDING_TYPE_BUFFER, &PersistBind };
 
-		if (PersistBuff)
+		if (InPersistBuff)
 		{
 #ifdef NNE_USE_D3D12_RESOURCES
-			PersistBind = DML_BUFFER_BINDING { PersistBuff, 0, PersistBuff->GetDesc().Width };
+			PersistBind = DML_BUFFER_BINDING { InPersistBuff, 0, InPersistBuff->GetDesc().Width };
 #else
-			PersistBind = MakeBind(PersistBuff);
+			PersistBind = MakeBind(InPersistBuff);
 #endif
 		}
 
@@ -166,22 +166,22 @@ public:
 		DML_BUFFER_BINDING			TempBind{};
 		DML_BINDING_DESC			TempBindDesc{ DML_BINDING_TYPE_BUFFER, &TempBind };
 
-		if (TempBuff)
+		if (InTempBuff)
 		{
 #ifdef NNE_USE_D3D12_RESOURCES
-			TempBind = { TempBuff, 0, TempBuff->GetDesc().Width };
+			TempBind = { InTempBuff, 0, InTempBuff->GetDesc().Width };
 #else
-			TempBind = MakeBind(TempBuff);
+			TempBind = MakeBind(InTempBuff);
 #endif
 			BindingTable->BindTemporaryResource(&TempBindDesc);
 		}
 	}
 
 #ifdef NNE_USE_D3D12_RESOURCES
-	void Bind(IDMLCompiledOperator* Op, TConstArrayView<FRHIBuffer*> InputBuffers, TConstArrayView<FRHIBuffer*> OutputBuffers, ID3D12Resource* PersistBuff = nullptr, ID3D12Resource* TempBuff = nullptr)
+	void Bind(IDMLCompiledOperator* Op, TConstArrayView<FRHIBuffer*> InputBuffers, TConstArrayView<FRHIBuffer*> OutputBuffers, ID3D12Resource* InPersistBuff = nullptr, ID3D12Resource* InTempBuff = nullptr)
 #else
 	
-	void Bind(IDMLCompiledOperator* Op, TConstArrayView<FRHIBuffer*> InputBuffers, TConstArrayView<FRHIBuffer*> OutputBuffers, FRHIBuffer* PersistBuff = nullptr, FRHIBuffer* TempBuff = nullptr)
+	void Bind(IDMLCompiledOperator* Op, TConstArrayView<FRHIBuffer*> InputBuffers, TConstArrayView<FRHIBuffer*> OutputBuffers, FRHIBuffer* InPersistBuff = nullptr, FRHIBuffer* InTempBuff = nullptr)
 #endif
 	{
 		Reset(Op);
@@ -202,12 +202,12 @@ public:
 		DML_BUFFER_BINDING			PersistBind;
 		DML_BINDING_DESC			PersistBindDesc{ DML_BINDING_TYPE_BUFFER, &PersistBind };
 
-		if (PersistBuff)
+		if (InPersistBuff)
 		{
 #ifdef NNE_USE_D3D12_RESOURCES
-			PersistBind = { PersistBuff, 0, PersistBuff->GetDesc().Width };
+			PersistBind = { InPersistBuff, 0, InPersistBuff->GetDesc().Width };
 #else
-			PersistBind =  MakeBind(PersistBuff);MakeBind(PersistBuff);
+			PersistBind =  MakeBind(InPersistBuff);MakeBind(InPersistBuff);
 #endif
 			BindingTable->BindPersistentResource(&PersistBindDesc);
 		}
@@ -215,13 +215,13 @@ public:
 		DML_BUFFER_BINDING			TempBind{};
 		DML_BINDING_DESC			TempBindDesc{ DML_BINDING_TYPE_BUFFER, &TempBind };
 
-		if (TempBuff)
+		if (InTempBuff)
 		{
 			
 #ifdef NNE_USE_D3D12_RESOURCES
-			TempBind = { TempBuff, 0, TempBuff->GetDesc().Width };
+			TempBind = { InTempBuff, 0, InTempBuff->GetDesc().Width };
 #else
-			TempBind = MakeBind(TempBuff);
+			TempBind = MakeBind(InTempBuff);
 #endif
 			BindingTable->BindTemporaryResource(&TempBindDesc);
 		}
@@ -387,12 +387,12 @@ private:
 
 public:
 
-	IDMLCompiledOperator* Compile(FDmlDeviceContext* DevCtx, const FGraphDesc& InGraph)
+	IDMLCompiledOperator* Compile(FDmlDeviceContext* InDevCtx, const FGraphDesc& InGraph)
 	{
-		IDMLDevice*				Device = DevCtx->Device;
+		IDMLDevice*				Device = InDevCtx->Device;
 		TComPtr<IDMLDevice1>	Device1;
 
-		Device1.FromQueryInterface(__uuidof(IDMLDevice1), DevCtx->Device);
+		Device1.FromQueryInterface(__uuidof(IDMLDevice1), InDevCtx->Device);
 		check(Device1);
 		if (!Device1)
 		{
