Add QA test for models that contain empty tensor initializers in ONNX

#rnx

diff --git a/Engine/Plugins/Experimental/NNE/Source/NNEQA/Resources/NNEQATestDesc.json b/Engine/Plugins/Experimental/NNE/Source/NNEQA/Resources/NNEQATestDesc.json
--- a/Engine/Plugins/Experimental/NNE/Source/NNEQA/Resources/NNEQATestDesc.json
+++ b/Engine/Plugins/Experimental/NNE/Source/NNEQA/Resources/NNEQATestDesc.json
@@ -129,6 +129,20 @@
 						{ "name": "NNERuntimeRDGHlsl", "skip": true, "comment": "Model not supported yet." }
 					],
 					"absolute_tolerance": 1e-4
+				},
+				{
+					"target": "ss",
+					"runtimes": [
+						{ "name": "NNERuntimeRDGHlsl", "skip": true, "comment": "Model not supported yet." }
+					],
+					"absolute_tolerance": 1e-4
+				},
+				{
+					"target": "ss_sim",
+					"runtimes": [
+						{ "name": "NNERuntimeRDGHlsl", "skip": true, "comment": "Model not supported yet." }
+					],
+					"absolute_tolerance": 1e-4
 				}
 			]
 		},
diff --git a/Engine/Plugins/Experimental/NNE/Source/NNEUtils/Private/NNEUtilsModelBuilderNNE.cpp b/Engine/Plugins/Experimental/NNE/Source/NNEUtils/Private/NNEUtilsModelBuilderNNE.cpp
--- a/Engine/Plugins/Experimental/NNE/Source/NNEUtils/Private/NNEUtilsModelBuilderNNE.cpp
+++ b/Engine/Plugins/Experimental/NNE/Source/NNEUtils/Private/NNEUtilsModelBuilderNNE.cpp
@@ -260,15 +260,15 @@ private:
 			if (Data)
 			{
 				Desc.Type = ENNEFormatTensorType::Initializer;
-				Desc.DataOffset = Format.TensorData.AddUninitialized(DataSize);
-				Desc.DataSize = DataSize;
 
-				FMemory::Memcpy(Format.TensorData.GetData() + Desc.DataOffset, Data, DataSize);
-			}
-			else
-			{
-				Desc.DataOffset = 0;
-				Desc.DataSize = 0;
+				// Handle empty data initializers, i.e. when DataSize is 0
+				if (DataSize)
+				{
+					Desc.DataOffset = Format.TensorData.AddUninitialized(DataSize);
+					Desc.DataSize = DataSize;
+
+					FMemory::Memcpy(Format.TensorData.GetData() + Desc.DataOffset, Data, DataSize);
+				}
 			}
 
 			Format.Tensors.Add(Desc);
