NNE - Add RDG Dml ArgMin and ArgMax

#rnx

diff --git a/Engine/Plugins/Experimental/NNE/Source/NNEQA/Resources/NNEQATestDesc.json b/Engine/Plugins/Experimental/NNE/Source/NNEQA/Resources/NNEQATestDesc.json
--- a/Engine/Plugins/Experimental/NNE/Source/NNEQA/Resources/NNEQATestDesc.json
+++ b/Engine/Plugins/Experimental/NNE/Source/NNEQA/Resources/NNEQATestDesc.json
@@ -529,6 +529,16 @@
 				{ "target": "ReduceSumSquare" }
 			]
 		},
+		{
+			"category": "ArgMaxMin",
+			"runtimes": [
+				{ "name": "NNERuntimeRDGHlsl", "skip": true }
+			],
+			"targets": [
+				{ "target": "ArgMax" },
+				{ "target": "ArgMin" }
+			]
+		},
 		{
 			"category": "Reshape",
 			"targets": [
@@ -1286,6 +1296,36 @@
 			"datasets": [
 				{ "inputs": [ { "shape": [ 3, 2, 2 ] } ], "outputs": [ { "shape": [ 3, 1, 2 ] } ] }
 			]
+		},
+		{ 
+			"name": "ArgMaxMin.Default",
+			"datasets": [
+				{ "inputs": [ { "shape": [ 3, 2, 2 ] } ], "outputs": [ { "shape": [ 1, 2, 2 ], "type": "int64" } ] }
+			]
+		},
+		{ 
+			"name": "ArgMaxMin.Axis.SelectLastIndex",
+			"datasets": [
+				{ "inputs": [ { "shape": [ 5, 7, 6 ] } ], "outputs": [ { "shape": [ 5, 1, 6 ], "type": "int64" } ] }
+			]
+		},
+		{ 
+			"name": "ArgMaxMin.Axis.Keepdims",
+			"datasets": [
+				{ "inputs": [ { "shape": [ 3, 2, 2 ] } ], "outputs": [ { "shape": [ 3, 1, 2 ], "type": "int64" } ] }
+			]
+		},
+		{ 
+			"name": "ArgMaxMin.Axis.DontKeepdims",
+			"datasets": [
+				{ "inputs": [ { "shape": [ 3, 2, 2 ] } ], "outputs": [ { "shape": [ 3, 2 ], "type": "int64" } ] }
+			]
+		},
+		{ 
+			"name": "ArgMaxMin.NegativeAxis.Keepdims",
+			"datasets": [
+				{ "inputs": [ { "shape": [ 3, 2, 2 ] } ], "outputs": [ { "shape": [ 3, 1, 2 ], "type": "int64" } ] }
+			]
 		}
 	],
 
@@ -1865,6 +1905,48 @@
 					]
 				}
 			]
+		},
+		{
+			"name": "ArgMaxMin.Axis.SelectLastIndex",
+			"attribute_maps": [
+				{
+					"attributes": [
+						{ "name": "axis", "value": 1, "type": "Int32" },
+						{ "name": "select_last_index", "value": 1, "type": "Int32" }
+					]
+				}
+			]
+		},
+		{
+			"name": "ArgMaxMin.Axis.Keepdims",
+			"attribute_maps": [
+				{
+					"attributes": [
+						{ "name": "axis", "value": 1, "type": "Int32" },
+						{ "name": "keepdims", "value": 1, "type": "Int32" }
+					]
+				}
+			]
+		},
+		{
+			"name": "ArgMaxMin.Axis.DontKeepdims",
+			"attribute_maps": [
+				{ "attributes": [ 
+						{ "name": "axis", "value": 1, "type": "Int32" },
+						{ "name": "keepdims", "value": 0, "type": "Int32" }
+					]
+				}
+			]
+		},
+		{
+			"name": "ArgMaxMin.NegativeAxis.Keepdims",
+			"attribute_maps": [
+				{ "attributes": [ 
+						{ "name": "axis", "value": -2, "type": "Int32" },
+						{ "name": "keepdims", "value": 1, "type": "Int32" }
+					]
+				}
+			]
 		}
 	]
 
diff --git a/Engine/Plugins/Experimental/NNE/Source/NNERuntimeRDG/Private/Dml/NNEDmlOperatorReduce.cpp b/Engine/Plugins/Experimental/NNE/Source/NNERuntimeRDG/Private/Dml/NNEDmlOperatorReduce.cpp
--- a/Engine/Plugins/Experimental/NNE/Source/NNERuntimeRDG/Private/Dml/NNEDmlOperatorReduce.cpp
+++ b/Engine/Plugins/Experimental/NNE/Source/NNERuntimeRDG/Private/Dml/NNEDmlOperatorReduce.cpp
@@ -55,11 +55,18 @@ public:
 		const int32 KeepDims = Attributes.GetValueOrDefault<int32>(TEXT("keepdims"), 1);
 		
 		TArray<int32>	Axes;
-		
-		const FNNEAttributeValue* AxesAttr = Attributes.GetAttributeValue(TEXT("axes"));
-		if (AxesAttr)
+
+		if constexpr (ReduceFunc == DML_REDUCE_FUNCTION_ARGMAX || ReduceFunc == DML_REDUCE_FUNCTION_ARGMIN)
 		{
-			Axes = AxesAttr->GetValue<TArray<int32>>();
+			Axes.Add(Attributes.GetValueOrDefault<int32>(TEXT("axis"), 0));
+		}
+		else
+		{
+			const FNNEAttributeValue* AxesAttr = Attributes.GetAttributeValue(TEXT("axes"));
+			if (AxesAttr)
+			{
+				Axes = AxesAttr->GetValue<TArray<int32>>();
+			}
 		}
 
 		HandleNegativeAxes(Axes, InputTensors[0].GetShape().Rank());
@@ -119,15 +126,41 @@ public:
 			return false;
 		}
 
-		DML_REDUCE_OPERATOR_DESC OpDesc{};
+		if constexpr (ReduceFunc == DML_REDUCE_FUNCTION_ARGMAX)
+        {
+            DML_ARGMAX_OPERATOR_DESC OpDesc;
+            OpDesc.AxisDirection = (DML_AXIS_DIRECTION) Attributes.GetValueOrDefault<int32>(TEXT("select_last_index"), 0);
+            OpDesc.InputTensor = &DmlInputTensorDesc.Desc;
+			OpDesc.OutputTensor = &DmlOutputTensorDesc.Desc;
+            OpDesc.Axes = DmlAxes.GetData();
+            OpDesc.AxisCount = (uint32) DmlAxes.Num();
+
+			return CreateOperator(Device, DML_OPERATOR_DESC { DML_OPERATOR_ARGMAX, &OpDesc });
+        }
+        else if constexpr (ReduceFunc == DML_REDUCE_FUNCTION_ARGMIN)
+        {
+            DML_ARGMIN_OPERATOR_DESC OpDesc;
+            OpDesc.AxisDirection = (DML_AXIS_DIRECTION) Attributes.GetValueOrDefault<int32>(TEXT("select_last_index"), 0);
+            OpDesc.InputTensor = &DmlInputTensorDesc.Desc;
+			OpDesc.OutputTensor = &DmlOutputTensorDesc.Desc;
+            OpDesc.Axes = DmlAxes.GetData();
+            OpDesc.AxisCount = (uint32) DmlAxes.Num();
+
+            return CreateOperator(Device, DML_OPERATOR_DESC { DML_OPERATOR_ARGMIN, &OpDesc });
+        }
+		else
+		{
+			DML_REDUCE_OPERATOR_DESC OpDesc{};
 
-		OpDesc.InputTensor = &DmlInputTensorDesc.Desc;
-		OpDesc.OutputTensor = &DmlOutputTensorDesc.Desc;
-		OpDesc.Function = ReduceFunc;
-		OpDesc.Axes = DmlAxes.GetData();
-		OpDesc.AxisCount = (uint32) DmlAxes.Num();
+			OpDesc.InputTensor = &DmlInputTensorDesc.Desc;
+			OpDesc.OutputTensor = &DmlOutputTensorDesc.Desc;
+			OpDesc.Function = ReduceFunc;
+			OpDesc.Axes = DmlAxes.GetData();
+			OpDesc.AxisCount = (uint32) DmlAxes.Num();
 
-		return CreateOperator(Device, DML_OPERATOR_DESC{ DML_OPERATOR_REDUCE, &OpDesc} );
+			return CreateOperator(Device, DML_OPERATOR_DESC{ DML_OPERATOR_REDUCE, &OpDesc} );
+		}
+		
 	}
 };
 
@@ -148,6 +181,8 @@ struct FOperatorDmlReduceRegistrator
 		OP(ReduceProd,		DML_REDUCE_FUNCTION_MULTIPLY);
 		OP(ReduceSum,		DML_REDUCE_FUNCTION_SUM);
 		OP(ReduceSumSquare,	DML_REDUCE_FUNCTION_SUM_SQUARE);
+		OP(ArgMax,			DML_REDUCE_FUNCTION_ARGMAX);
+		OP(ArgMin,			DML_REDUCE_FUNCTION_ARGMIN);
 	}
 };
 
