NNE Skip ORTCuda tests when NNERuntimeORTCuda model is not available.

What: We still want to be able to get green NNE tests on a system where Cudnn and Cuda are not installed
How: Skip cuda tests if cudnn is not installed. Log will show ORTCuda tests as skipped.

Tests: NNE smoke tests on win64 editor

#rnx

diff --git a/Engine/Plugins/Experimental/NNE/Source/NNEQA/Private/NNEQAParametricTest.cpp b/Engine/Plugins/Experimental/NNE/Source/NNEQA/Private/NNEQAParametricTest.cpp
--- a/Engine/Plugins/Experimental/NNE/Source/NNEQA/Private/NNEQAParametricTest.cpp
+++ b/Engine/Plugins/Experimental/NNE/Source/NNEQA/Private/NNEQAParametricTest.cpp
@@ -13,6 +13,7 @@
 #include "NNEUtilsModelBuilder.h"
 #include "UObject/Class.h"
 #include "UObject/ReflectedTypeAccessors.h"
+#include "GenericPlatform/GenericPlatformMisc.h"
 
 #if WITH_DEV_AUTOMATION_TESTS
 #include "Misc/AutomationTest.h"
@@ -56,9 +57,28 @@ public:
 		AddTestFromCategory(NNEBaseTestPath + TEXT(".Model."), ModelTestCategories, InputOutputSets, AttributeSets);
 		AddTestFromCategory(NNEBaseTestPath + TEXT(".Operator."), OperatorTestCategories, InputOutputSets, AttributeSets);
 
+		if (ShouldSkipORTCudaTests())
+		{
+			for (FTestSetup& Test : TestSetups)
+			{
+				Test.AutomationExcludedRuntime.AddUnique(TEXT("NNERuntimeORTCuda"));
+			}
+		}
+
 		return true;
 	}
 private:
+	static bool ShouldSkipORTCudaTests()
+	{
+#if PLATFORM_WINDOWS
+		//Heuristic: if Cudnn is installed it should be in the PATH and we run the test.
+		if (FPlatformMisc::GetEnvironmentVariable(TEXT("PATH")).Contains("CUDNN"))
+		{
+			return false;
+		}
+#endif
+		return true;
+	}
 		
 	static void ApplyEpsilons(FTests::FTestSetup& TestSetup, const Json::FTestConfigTarget& TestTarget)
 	{
@@ -719,11 +739,16 @@ static FAutoConsoleCommand TestReloadCommand(
 	FConsoleCommandWithArgsDelegate::CreateStatic(
 		[](const TArray< FString >& Args)
 		{
-			ParametricTests.ReloadTestDescriptionsFromJson();
+			bool bResult = ParametricTests.ReloadTestDescriptionsFromJson();
 				
 			#if WITH_DEV_AUTOMATION_TESTS
 				ParametricTestAutomationRegistry.Refresh();
 			#endif
+
+			if (!bResult)
+			{
+				UE_LOG(LogNNE, Warning, TEXT("Failed reloading NNE test from Json."));
+			}
 		}
 	)
 );
