NNE - Add RDG Dml CumSum

#rnx

diff --git a/Engine/Plugins/Experimental/NNE/Source/NNEQA/Resources/NNEQATestDesc.json b/Engine/Plugins/Experimental/NNE/Source/NNEQA/Resources/NNEQATestDesc.json
--- a/Engine/Plugins/Experimental/NNE/Source/NNEQA/Resources/NNEQATestDesc.json
+++ b/Engine/Plugins/Experimental/NNE/Source/NNEQA/Resources/NNEQATestDesc.json
@@ -623,6 +623,15 @@
 			"runtimes": [
 				{ "name": "NNERuntimeRDGHlsl", "skip": true, "comment": "Not supported" }
 			]
+		},
+		{
+			"category": "CumSum",
+			"targets": [
+				{ "target": "CumSum" }
+			],
+			"runtimes": [
+				{ "name": "NNERuntimeRDGHlsl", "skip": true, "comment": "Not supported" }
+			]
 		}
 	],
 
@@ -1318,6 +1327,30 @@
 				{ "inputs": [ { "shape": [ 1, 1, 6, 2 ] } ], "outputs": [ { "shape": [ 1, 1, 6, 1 ] }, { "shape": [ 1, 1, 6, 1 ] } ] }
 			]
 		},
+		{
+			"name": "CumSum.1D",
+			"datasets": [
+				{ "inputs": [ { "shape": [ 5 ] } ], "weights": [ { "shape": [ 1 ], "source": [ "0" ], "type": "int64" } ], "outputs": [ { "shape": [ 5 ] } ] }
+			]
+		},
+		{
+			"name": "CumSum.2DAxis0",
+			"datasets": [
+				{ "inputs": [ { "shape": [ 2, 3 ] } ], "weights": [ { "shape": [ 1 ], "source": [ "0" ], "type": "int64" } ], "outputs": [ { "shape": [ 2, 3 ] } ] }
+			]
+		},
+		{
+			"name": "CumSum.2DAxis1",
+			"datasets": [
+				{ "inputs": [ { "shape": [ 3, 4 ] } ], "weights": [ { "shape": [ 1 ], "source": [ "1" ], "type": "int64" } ], "outputs": [ { "shape": [ 3, 4 ] } ] }
+			]
+		},
+		{
+			"name": "CumSum.2DNegativeAxis",
+			"datasets": [
+				{ "inputs": [ { "shape": [ 3, 2 ] } ], "weights": [ { "shape": [ 1 ], "source": [ "-1" ], "type": "int64" } ], "outputs": [ { "shape": [ 3, 2 ] } ] }
+			]
+		},
 		{ 
 			"name": "Reduce.Default",
 			"datasets": [
@@ -1920,6 +1953,29 @@
 				}
 			]
 		},
+		{
+			"name": "CumSum.1D",
+			"attribute_maps": [
+				{ 
+					"attributes": [ 
+						{ "name": "reverse", "value": 0, "type": "Int32" }
+					]
+				},
+				{ 
+					"attributes": [ 
+						{ "name": "reverse", "value": 1, "type": "Int32" }
+					]
+				}
+			]
+		},
+		{
+			"name": "CumSum.1D.Exclusive",
+			"multiply_with_attribute_sets": [ "CumSum.1D" ],
+			"attribute_maps": [
+				{ "attributes": [ { "name": "exclusive", "value": 0, "type": "Int32" } ] },
+				{ "attributes": [ { "name": "exclusive", "value": 1, "type": "Int32" } ] }
+			]
+		},
 		{
 			"name": "Reduce.Axes.Keepdims",
 			"attribute_maps": [
diff --git a/Engine/Plugins/Experimental/NNE/Source/NNERuntimeRDG/Private/Dml/NNEDmlOperator.h b/Engine/Plugins/Experimental/NNE/Source/NNERuntimeRDG/Private/Dml/NNEDmlOperator.h
--- a/Engine/Plugins/Experimental/NNE/Source/NNERuntimeRDG/Private/Dml/NNEDmlOperator.h
+++ b/Engine/Plugins/Experimental/NNE/Source/NNERuntimeRDG/Private/Dml/NNEDmlOperator.h
@@ -61,6 +61,17 @@ namespace DmlUtil
 	using FSmallIntArray = TArray<int32, TInlineAllocator<NNECore::FTensorShape::MaxRank>>;
 	using FSmallUIntArray = TArray<uint32, TInlineAllocator<NNECore::FTensorShape::MaxRank>>;
 
+	template<typename InputType, typename OutputType>
+	inline bool IsOverflowing(InputType Input)
+	{
+		OutputType Output = static_cast<OutputType>(Input);
+		if(Input != static_cast<InputType>(Output))
+		{
+			return true;
+		}
+		return false;
+	}
+
 	template<typename InputType, typename OutputType>
 	inline bool ConvertArrayViewNoOverflow(TConstArrayView<InputType> InputView, TArrayView<OutputType>& OutputView)
 	{
diff --git a/Engine/Plugins/Experimental/NNE/Source/NNERuntimeRDG/Private/Dml/NNEDmlOperatorCumSum.cpp b/Engine/Plugins/Experimental/NNE/Source/NNERuntimeRDG/Private/Dml/NNEDmlOperatorCumSum.cpp
new file mode 100644
--- /dev/null
+++ b/Engine/Plugins/Experimental/NNE/Source/NNERuntimeRDG/Private/Dml/NNEDmlOperatorCumSum.cpp
@@ -0,0 +1,104 @@
+// Copyright Epic Games, Inc. All Rights Reserved.
+
+#ifdef NNE_USE_DIRECTML
+
+#include "NNEDmlOperator.h"
+#include "NNEDmlOperatorUtils.h"
+
+namespace UE::NNERuntimeRDG::Private::Dml
+{
+
+class FOperatorDmlCumSum : public FOperatorDml
+{
+
+public:
+
+	static FOperatorDml* Create()
+	{
+		return new FOperatorDmlCumSum();
+	}
+
+	//
+	//
+	//
+	virtual bool Initialize(IDMLDevice* Device, TArrayView<const NNECore::Internal::FTensor> InputTensors, TArrayView<const NNECore::Internal::FTensor> OutputTensors, const NNECore::FAttributeMap& Attributes) override
+	{
+		check(InputTensors.Num() == 2);
+		check(OutputTensors.Num() == 1);
+
+		ConstantCPUInputs.Add(1);
+
+		UINT Axis; 
+
+        // axis tensor must be constant
+        check(InputTensors[1].HasPreparedData());
+
+		switch(InputTensors[1].GetDataType())
+        {
+        case ENNETensorDataType::Int32:
+			{
+				TConstArrayView<int32> TensorContent = InputTensors[1].GetPreparedData<int32>();
+				if(TensorContent.Num() != 1)
+				{
+					UE_LOG(LogNNE, Error, TEXT("axis tensor should be 0-D"));
+					return false;
+				}
+				Axis = (UINT) HandleNegativeAxis(TensorContent[0], InputTensors[0].GetShape().Rank());
+			}
+            break;
+        case ENNETensorDataType::Int64:
+			{
+				TConstArrayView<int64> TensorContent = InputTensors[1].GetPreparedData<int64>();
+				if(TensorContent.Num() != 1)
+				{
+					UE_LOG(LogNNE, Error, TEXT("axis tensor should be 0-D"));
+					return false;
+				}
+				if(DmlUtil::IsOverflowing<int64, int32>(TensorContent[0]))
+				{
+					UE_LOG(LogNNE, Error, TEXT("axis number is overflowing"));
+					return false;
+				}
+				Axis = (UINT) HandleNegativeAxis(TensorContent[0], InputTensors[0].GetShape().Rank());
+			}
+            break;
+        default:
+            UE_LOG(LogNNE, Error, TEXT("axis tensor has invalid data type"));
+			return false;
+        };
+
+		const NNECore::Internal::FTensor& InputTensorDesc = InputTensors[0];
+		const NNECore::Internal::FTensor& OutputTensorDesc = OutputTensors[0];
+
+		DmlUtil::FTensorDesc	DmlInputTensorDesc{};
+		DmlUtil::FTensorDesc	DmlOutputTensorDesc{};
+
+		if (!InitDmlTensorDesc(DmlInputTensorDesc, InputTensorDesc))
+		{
+			UE_LOG(LogNNE, Warning, TEXT("Failed to initialize input tensor for DML inference"));
+			return false;
+		}
+		if (!InitDmlTensorDesc(DmlOutputTensorDesc, OutputTensorDesc))
+		{
+			UE_LOG(LogNNE, Warning, TEXT("Failed to initialize output tensor for DML inference"));
+			return false;
+		}
+
+		DML_CUMULATIVE_SUMMATION_OPERATOR_DESC DmlCumSumOpDesc{};
+
+		DmlCumSumOpDesc.InputTensor = &DmlInputTensorDesc.Desc;
+		DmlCumSumOpDesc.OutputTensor = &DmlOutputTensorDesc.Desc;
+		DmlCumSumOpDesc.Axis = Axis;
+		DmlCumSumOpDesc.AxisDirection = (DML_AXIS_DIRECTION) Attributes.GetValueOrDefault<int32>(TEXT("reverse"), 0);
+		DmlCumSumOpDesc.HasExclusiveSum = (bool) Attributes.GetValueOrDefault<int32>(TEXT("exclusive"), 0);
+
+		return CreateOperator(Device, DML_OPERATOR_DESC{ DML_OPERATOR_CUMULATIVE_SUMMATION, &DmlCumSumOpDesc });
+	}
+};
+
+// Register CumSum operator on Module startup
+NNE_DML_REGISTER_OP(CumSum)
+
+} // namespace UE::NNERuntimeRDG::Private::Dml
+
+#endif // NNE_USE_DIRECTML
