NNE - ADD RDG Dml MaxUnpool

#rnx

diff --git a/Engine/Plugins/Experimental/NNE/Source/NNEQA/Resources/NNEQATestDesc.json b/Engine/Plugins/Experimental/NNE/Source/NNEQA/Resources/NNEQATestDesc.json
--- a/Engine/Plugins/Experimental/NNE/Source/NNEQA/Resources/NNEQATestDesc.json
+++ b/Engine/Plugins/Experimental/NNE/Source/NNEQA/Resources/NNEQATestDesc.json
@@ -615,6 +615,15 @@
 				{ "name": "NNERuntimeRDGHlsl", "skip": true, "comment": "Not supported" }
 			]
 		},
+		{
+			"category": "MaxUnpool",
+			"targets": [
+				{ "target": "MaxUnpool" }
+			],
+			"runtimes": [
+				{ "name": "NNERuntimeRDGHlsl", "skip": true, "comment": "Not supported" }
+			]
+		},
 		{
 			"category": "Split",
 			"targets": [
@@ -1315,6 +1324,22 @@
 				{ "inputs": [ { "shape": [ 2, 3, 128, 256, 64 ] } ], "outputs": [ { "shape": [ 2, 3, 1, 1, 1 ] } ] }
 			]
 		},
+		{
+			"name": "MaxUnpool.WithOutputShape",
+			"datasets": [
+				{
+					"inputs": [ { "shape": [ 1, 1, 2, 2 ] }, { "shape": [ 1, 1, 2, 2 ], "source": [ "5", "7", "13", "15" ], "type": "int64" } ],
+					"weights": [ { "shape": [ 4 ], "source": [ "1", "1", "5", "5" ], "type": "int64" } ], 
+					"outputs": [ { "shape": [ 1, 1, 5, 5 ] } ] 
+				}
+			]
+		},
+		{
+			"name": "MaxUnpool.WithoutOutputShape",
+			"datasets": [
+				{ "inputs": [ { "shape": [ 1, 1, 2, 2 ] }, { "shape": [ 1, 1, 2, 2 ], "source": [ "5", "7", "13", "15" ], "type": "int64" } ], "outputs": [ { "shape": [ 1, 1, 4, 4 ] } ] }
+			]
+		},
 		{
 			"name": "Split.Axis2",
 			"datasets": [
@@ -1933,6 +1958,26 @@
 				}
 			]
 		},
+		{
+			"name": "MaxUnpool.WithOutputShape",
+			"attribute_maps": [
+				{ "attributes": [ 
+						{ "name": "kernel_shape", "value": [2, 2], "type": "Int32Array" },
+						{ "name": "strides", "value": [2, 2], "type": "Int32Array" }
+					]
+				}
+			]
+		},
+		{
+			"name": "MaxUnpool.WithoutOutputShape",
+			"attribute_maps": [
+				{ "attributes": [ 
+						{ "name": "kernel_shape", "value": [2, 2], "type": "Int32Array" },
+						{ "name": "strides", "value": [2, 2], "type": "Int32Array" }
+					]
+				}
+			]
+		},
 		{
 			"name": "Split.Axis2",
 			"attribute_maps": [
diff --git a/Engine/Plugins/Experimental/NNE/Source/NNERuntimeRDG/Private/Dml/NNEDmlOperatorMaxUnpool.cpp b/Engine/Plugins/Experimental/NNE/Source/NNERuntimeRDG/Private/Dml/NNEDmlOperatorMaxUnpool.cpp
new file mode 100644
--- /dev/null
+++ b/Engine/Plugins/Experimental/NNE/Source/NNERuntimeRDG/Private/Dml/NNEDmlOperatorMaxUnpool.cpp
@@ -0,0 +1,76 @@
+// Copyright Epic Games, Inc. All Rights Reserved.
+
+#ifdef NNE_USE_DIRECTML
+
+#include "NNEDmlOperator.h"
+#include "NNEDmlOperatorUtils.h"
+
+namespace UE::NNERuntimeRDG::Private::Dml
+{
+
+class FOperatorDmlMaxUnpool : public FOperatorDml
+{
+
+public:
+
+	static FOperatorDml* Create()
+	{
+		return new FOperatorDmlMaxUnpool();
+	}
+
+	//
+	//
+	//
+	virtual bool Initialize(IDMLDevice* Device, TArrayView<const NNECore::Internal::FTensor> InputTensors, TArrayView<const NNECore::Internal::FTensor> OutputTensors, const NNECore::FAttributeMap& Attributes) override
+	{
+		check(InputTensors.Num() == 2 || InputTensors.Num() == 3);
+		check(OutputTensors.Num() == 1);
+
+		if(InputTensors.Num() == 3)
+		{
+			ConstantCPUInputs.Add(2);
+		}
+
+		const NNECore::Internal::FTensor& InputTensorDesc = InputTensors[0];
+		const NNECore::Internal::FTensor& IndicesTensorDesc = InputTensors[1];
+		const NNECore::Internal::FTensor& OutputTensorDesc = OutputTensors[0];
+
+		DmlUtil::FTensorDesc	DmlInputTensorDesc;
+		DmlUtil::FTensorDesc	DmlIndicesTensorDesc;
+		DmlUtil::FTensorDesc	DmlOutputTensorDesc;
+
+		if (!DmlInputTensorDesc.InitFromTensor(InputTensorDesc, 4))
+		{
+			UE_LOG(LogNNE, Warning, TEXT("Failed to initialize input tensor for DML inference"));
+			return false;
+		}
+		if (!DmlIndicesTensorDesc.InitFromTensor(IndicesTensorDesc, 4))
+		{
+			UE_LOG(LogNNE, Warning, TEXT("Failed to initialize indices tensor for DML inference"));
+			return false;
+		}
+		check(DmlIndicesTensorDesc.BuffDesc.DataType == DML_TENSOR_DATA_TYPE::DML_TENSOR_DATA_TYPE_INT64);
+		// Cast IndicesTensor from int64 to uint64 due to differences in representation between ONNX and DML formats.
+		DmlIndicesTensorDesc.BuffDesc.DataType = DML_TENSOR_DATA_TYPE::DML_TENSOR_DATA_TYPE_UINT64;
+		if (!DmlOutputTensorDesc.InitFromTensor(OutputTensorDesc, 4))
+		{
+			UE_LOG(LogNNE, Warning, TEXT("Failed to initialize output tensor for DML inference"));
+			return false;
+		}
+
+		DML_MAX_UNPOOLING_OPERATOR_DESC DmlMaxUnpoolOpDesc{};
+
+		DmlMaxUnpoolOpDesc.InputTensor = &DmlInputTensorDesc.Desc;
+		DmlMaxUnpoolOpDesc.IndicesTensor = &DmlIndicesTensorDesc.Desc;
+		DmlMaxUnpoolOpDesc.OutputTensor = &DmlOutputTensorDesc.Desc;
+
+		return CreateOperator(Device, DML_OPERATOR_DESC{ DML_OPERATOR_MAX_UNPOOLING, &DmlMaxUnpoolOpDesc });
+	}
+};
+
+// Register MaxUnpool operator on Module startup
+NNE_DML_REGISTER_OP(MaxUnpool)
+
+} // namespace UE::NNERuntimeRDG::Private::Dml
+
+#endif // NNE_USE_DIRECTML
