NNE - Add RDG Dml DepthToSpace&SpaceToDepth

#rnx

diff --git a/Engine/Plugins/Experimental/NNE/Source/NNEQA/Resources/NNEQATestDesc.json b/Engine/Plugins/Experimental/NNE/Source/NNEQA/Resources/NNEQATestDesc.json
--- a/Engine/Plugins/Experimental/NNE/Source/NNEQA/Resources/NNEQATestDesc.json
+++ b/Engine/Plugins/Experimental/NNE/Source/NNEQA/Resources/NNEQATestDesc.json
@@ -641,6 +641,24 @@
 			"runtimes": [
 				{ "name": "NNERuntimeRDGHlsl", "skip": true, "comment": "Not supported" }
 			]
+		},
+		{
+			"category": "DepthToSpace",
+			"targets": [
+				{ "target": "DepthToSpace" }
+			],
+			"runtimes": [
+				{ "name": "NNERuntimeRDGHlsl", "skip": true, "comment": "Not supported" }
+			]
+		},
+		{
+			"category": "SpaceToDepth",
+			"targets": [
+				{ "target": "SpaceToDepth" }
+			],
+			"runtimes": [
+				{ "name": "NNERuntimeRDGHlsl", "skip": true, "comment": "Not supported" }
+			]
 		}
 	],
 
@@ -1376,6 +1394,30 @@
 				{ "inputs": [ { "shape": [ 3, 2 ] } ], "weights": [ { "shape": [ 1 ], "source": [ "-1" ], "type": "int64" } ], "outputs": [ { "shape": [ 3, 2 ] } ] }
 			]
 		},
+		{
+			"name": "DepthToSpace.CRD",
+			"datasets": [
+				{ "inputs": [ { "shape": [ 1, 8, 2, 3 ] } ], "outputs": [ { "shape": [ 1, 2, 4, 6 ] } ] }
+			]
+		},
+		{
+			"name": "DepthToSpace.DCR",
+			"datasets": [
+				{ "inputs": [ { "shape": [ 1, 8, 2, 3 ] } ], "outputs": [ { "shape": [ 1, 2, 4, 6 ] } ] }
+			]
+		},
+		{
+			"name": "SpaceToDepth.BlockSize2",
+			"datasets": [
+				{ "inputs": [ { "shape": [ 1, 1, 4, 6 ] } ], "outputs": [ { "shape": [ 1, 4, 2, 3 ] } ] }
+			]
+		},
+		{
+			"name": "SpaceToDepth.BlockSize4",
+			"datasets": [
+				{ "inputs": [ { "shape": [ 2, 1, 12, 20 ] } ], "outputs": [ { "shape": [ 2, 16, 3, 5 ] } ] }
+			]
+		},
 		{ 
 			"name": "Reduce.Default",
 			"datasets": [
@@ -2021,6 +2063,40 @@
 				{ "attributes": [ { "name": "exclusive", "value": 1, "type": "Int32" } ] }
 			]
 		},
+		{
+			"name": "DepthToSpace.CRD",
+			"attribute_maps": [
+				{ "attributes": [ 
+					{ "name": "blocksize", "value": 2, "type": "Int32" },
+					{ "name": "mode", "value": "CRD", "type": "String" }
+					] }
+			]
+		},
+		{
+			"name": "DepthToSpace.DCR",
+			"attribute_maps": [
+				{ "attributes": [ 
+					{ "name": "blocksize", "value": 2, "type": "Int32" },
+					{ "name": "mode", "value": "DCR", "type": "String" }
+					] }
+			]
+		},
+		{
+			"name": "SpaceToDepth.BlockSize2",
+			"attribute_maps": [
+				{ "attributes": [ 
+					{ "name": "blocksize", "value": 2, "type": "Int32" }
+					] }
+			]
+		},
+		{
+			"name": "SpaceToDepth.BlockSize4",
+			"attribute_maps": [
+				{ "attributes": [ 
+					{ "name": "blocksize", "value": 4, "type": "Int32" }
+					] }
+			]
+		},
 		{
 			"name": "Reduce.Axes.Keepdims",
 			"attribute_maps": [
diff --git a/Engine/Plugins/Experimental/NNE/Source/NNERuntimeRDG/Private/Dml/NNEDmlOperatorDepthToSpace.cpp b/Engine/Plugins/Experimental/NNE/Source/NNERuntimeRDG/Private/Dml/NNEDmlOperatorDepthToSpace.cpp
new file mode 100644
--- /dev/null
+++ b/Engine/Plugins/Experimental/NNE/Source/NNERuntimeRDG/Private/Dml/NNEDmlOperatorDepthToSpace.cpp
@@ -0,0 +1,88 @@
+// Copyright Epic Games, Inc. All Rights Reserved.
+
+#ifdef NNE_USE_DIRECTML
+
+#include "NNEDmlOperator.h"
+#include "NNEDmlOperatorUtils.h"
+
+namespace UE::NNERuntimeRDG::Private::Dml
+{
+
+class FOperatorDmlDepthToSpace : public FOperatorDml
+{
+
+	static DML_DEPTH_SPACE_ORDER SpaceOrderFromModeString(FStringView StringVal)
+	{
+		if (FCString::Stricmp(StringVal.GetData(), TEXT("CRD")) == 0)
+		{
+			return DML_DEPTH_SPACE_ORDER_COLUMN_ROW_DEPTH;
+		}
+		else
+		{
+			return DML_DEPTH_SPACE_ORDER_DEPTH_COLUMN_ROW;
+		}
+	}
+
+public:
+
+	static FOperatorDml* Create()
+	{
+		return new FOperatorDmlDepthToSpace();
+	}
+
+	//
+	//
+	//
+	virtual bool Initialize(IDMLDevice* Device, TArrayView<const NNECore::Internal::FTensor> InputTensors, TArrayView<const NNECore::Internal::FTensor> OutputTensors, const NNECore::FAttributeMap& Attributes) override
+	{
+		check(InputTensors.Num() == 1);
+		check(OutputTensors.Num() == 1);
+
+		int32 BlockSize;
+
+		const FNNEAttributeValue* BlockSizeAttr = Attributes.GetAttributeValue(TEXT("blocksize"));
+		if (BlockSizeAttr)
+		{
+			BlockSize = BlockSizeAttr->GetValue<int32>();
+		}
+		else
+		{
+			UE_LOG(LogNNE, Error, TEXT("blocksize attribute is required"));
+			return false;
+		}
+
+
+		const NNECore::Internal::FTensor& InputTensorDesc = InputTensors[0];
+		const NNECore::Internal::FTensor& OutputTensorDesc = OutputTensors[0];
+
+		DmlUtil::FTensorDesc	DmlInputTensorDesc{};
+		DmlUtil::FTensorDesc	DmlOutputTensorDesc{};
+
+		if (!DmlInputTensorDesc.InitFromTensor(InputTensorDesc, 4))
+		{
+			UE_LOG(LogNNE, Warning, TEXT("Failed to initialize input tensor for DML inference"));
+			return false;
+		}
+		if (!DmlOutputTensorDesc.InitFromTensor(OutputTensorDesc, 4))
+		{
+			UE_LOG(LogNNE, Warning, TEXT("Failed to initialize output tensor for DML inference"));
+			return false;
+		}
+
+		DML_DEPTH_TO_SPACE1_OPERATOR_DESC DmlDepthToSpaceOpDesc{};
+
+		DmlDepthToSpaceOpDesc.InputTensor = &DmlInputTensorDesc.Desc;
+		DmlDepthToSpaceOpDesc.OutputTensor = &DmlOutputTensorDesc.Desc;
+		DmlDepthToSpaceOpDesc.BlockSize = BlockSize;
+		DmlDepthToSpaceOpDesc.Order = SpaceOrderFromModeString(Attributes.GetValueOrDefault<FString>(TEXT("mode"), FString(TEXT("DCR"))));
+
+		return CreateOperator(Device, DML_OPERATOR_DESC{ DML_OPERATOR_DEPTH_TO_SPACE1, &DmlDepthToSpaceOpDesc });
+	}
+};
+
+// Register DepthToSpace operator on Module startup
+NNE_DML_REGISTER_OP(DepthToSpace)
+
+} // namespace UE::NNERuntimeRDG::Private::Dml
+
+#endif // NNE_USE_DIRECTML
diff --git a/Engine/Plugins/Experimental/NNE/Source/NNERuntimeRDG/Private/Dml/NNEDmlOperatorSpaceToDepth.cpp b/Engine/Plugins/Experimental/NNE/Source/NNERuntimeRDG/Private/Dml/NNEDmlOperatorSpaceToDepth.cpp
new file mode 100644
--- /dev/null
+++ b/Engine/Plugins/Experimental/NNE/Source/NNERuntimeRDG/Private/Dml/NNEDmlOperatorSpaceToDepth.cpp
@@ -0,0 +1,77 @@
+// Copyright Epic Games, Inc. All Rights Reserved.
+
+#ifdef NNE_USE_DIRECTML
+
+#include "NNEDmlOperator.h"
+#include "NNEDmlOperatorUtils.h"
+
+namespace UE::NNERuntimeRDG::Private::Dml
+{
+
+class FOperatorDmlSpaceToDepth : public FOperatorDml
+{
+
+
+public:
+
+	static FOperatorDml* Create()
+	{
+		return new FOperatorDmlSpaceToDepth();
+	}
+
+	//
+	//
+	//
+	virtual bool Initialize(IDMLDevice* Device, TArrayView<const NNECore::Internal::FTensor> InputTensors, TArrayView<const NNECore::Internal::FTensor> OutputTensors, const NNECore::FAttributeMap& Attributes) override
+	{
+		check(InputTensors.Num() == 1);
+		check(OutputTensors.Num() == 1);
+
+		int32 BlockSize;
+
+		const FNNEAttributeValue* BlockSizeAttr = Attributes.GetAttributeValue(TEXT("blocksize"));
+		if (BlockSizeAttr)
+		{
+			BlockSize = BlockSizeAttr->GetValue<int32>();
+		}
+		else
+		{
+			UE_LOG(LogNNE, Error, TEXT("blocksize attribute is required"));
+			return false;
+		}
+
+
+		const NNECore::Internal::FTensor& InputTensorDesc = InputTensors[0];
+		const NNECore::Internal::FTensor& OutputTensorDesc = OutputTensors[0];
+
+		DmlUtil::FTensorDesc	DmlInputTensorDesc{};
+		DmlUtil::FTensorDesc	DmlOutputTensorDesc{};
+
+		if (!DmlInputTensorDesc.InitFromTensor(InputTensorDesc, 4))
+		{
+			UE_LOG(LogNNE, Warning, TEXT("Failed to initialize input tensor for DML inference"));
+			return false;
+		}
+		if (!DmlOutputTensorDesc.InitFromTensor(OutputTensorDesc, 4))
+		{
+			UE_LOG(LogNNE, Warning, TEXT("Failed to initialize output tensor for DML inference"));
+			return false;
+		}
+
+		DML_SPACE_TO_DEPTH1_OPERATOR_DESC DmlSpaceToDepthOpDesc{};
+
+		DmlSpaceToDepthOpDesc.InputTensor = &DmlInputTensorDesc.Desc;
+		DmlSpaceToDepthOpDesc.OutputTensor = &DmlOutputTensorDesc.Desc;
+		DmlSpaceToDepthOpDesc.BlockSize = BlockSize;
+		DmlSpaceToDepthOpDesc.Order = DML_DEPTH_SPACE_ORDER_DEPTH_COLUMN_ROW;
+
+		return CreateOperator(Device, DML_OPERATOR_DESC{ DML_OPERATOR_SPACE_TO_DEPTH1, &DmlSpaceToDepthOpDesc });
+	}
+};
+
+// Register SpaceToDepth operator on Module startup
+NNE_DML_REGISTER_OP(SpaceToDepth)
+
+} // namespace UE::NNERuntimeRDG::Private::Dml
+
+#endif // NNE_USE_DIRECTML
