Use FTensorDescDml and only member functions when working with DirectML tensor descriptors

#rnx

diff --git a/Engine/Plugins/Experimental/NNE/Source/NNERuntimeRDG/Private/Dml/NNEDmlOperator.cpp b/Engine/Plugins/Experimental/NNE/Source/NNERuntimeRDG/Private/Dml/NNEDmlOperator.cpp
--- a/Engine/Plugins/Experimental/NNE/Source/NNERuntimeRDG/Private/Dml/NNEDmlOperator.cpp
+++ b/Engine/Plugins/Experimental/NNE/Source/NNERuntimeRDG/Private/Dml/NNEDmlOperator.cpp
@@ -5,134 +5,129 @@
 #include "NNEDmlOperator.h"
 #include "ID3D12DynamicRHI.h"
 
+#include "Algo/Compare.h"
+
 namespace UE::NNERuntimeRDG::Private::Dml
 {
 
-namespace DmlUtil
+namespace Util
 {
 
-	bool FTensorDesc::InitFromTensor(const NNECore::Internal::FTensor& Tensor, int32 MinTensorRank, TConstArrayView<uint32> BroadcastShape, TConstArrayView<uint32> CustomShape)
-	{
-		Reset();
+static bool IsSameShape(TConstArrayView<uint32> Left, TConstArrayView<uint32> Right)
+{
+	return Algo::Compare(Left, Right);
+}
 
-		DML_TENSOR_DATA_TYPE DmlDataType = DmlUtil::GetTensorDataType(Tensor.GetDataType());
+bool IsSameShape(const NNECore::Internal::FTensor& Left, const NNECore::Internal::FTensor& Right)
+{
+	return IsSameShape(Left.GetShape().GetData(), Right.GetShape().GetData());
+}
 
-		if (DmlDataType == DML_TENSOR_DATA_TYPE_UNKNOWN)
-		{
-			return false;
-		}
+bool IsSameShape(const FTensorDescDml& Left, const FTensorDescDml& Right)
+{
+	return IsSameShape(Left.GetSizes(), Right.GetSizes());
+}
 
-		TConstArrayView<uint32> InShape = CustomShape.IsEmpty() ? Tensor.GetShape().GetData() : CustomShape;
 
-		ElemSizeInBytes = Tensor.GetElemByteSize();
+DML_TENSOR_DATA_TYPE GetTensorDataType(ENNETensorDataType DataType)
+{
+	switch (DataType)
+	{
+	case ENNETensorDataType::Double:
+		return DML_TENSOR_DATA_TYPE_FLOAT64;
 
-		if (BroadcastShape.IsEmpty())
-		{
-			SetShape(InShape, MinTensorRank);
-		}
-		else
-		{
-			SetShapeAndStrides(InShape, BroadcastShape);
-		}
+	case ENNETensorDataType::Float:
+		return DML_TENSOR_DATA_TYPE_FLOAT32;
 
-		//Note: We should support tensor padding using strides defined in FTensorDesc
-		//DmlUtil::SetTensorStrides(DmlTensorDesc, TensorDesc.Strides);
+	case ENNETensorDataType::Half:
+		return DML_TENSOR_DATA_TYPE_FLOAT16;
 
-		Update(DmlDataType, Tensor.HasPreparedData());
+	case ENNETensorDataType::UInt64:
+		return DML_TENSOR_DATA_TYPE_UINT64;
 
-		return true;
-	}
+	case ENNETensorDataType::UInt32:
+		return DML_TENSOR_DATA_TYPE_UINT32;
 
-	bool FTensorDesc::InitFromTensor1D(const NNECore::Internal::FTensor& Tensor, int32 Rank)
-	{
-		Reset();
+	case ENNETensorDataType::UInt16:
+		return DML_TENSOR_DATA_TYPE_UINT16;
 
-		if (Tensor.GetShape().Rank() != 1)
-		{
-			return false;
-		}
+	case ENNETensorDataType::UInt8:
+		return DML_TENSOR_DATA_TYPE_UINT8;
 
-		DML_TENSOR_DATA_TYPE DmlDataType = DmlUtil::GetTensorDataType(Tensor.GetDataType());
+	case ENNETensorDataType::Int64:
+		return DML_TENSOR_DATA_TYPE_INT64;
 
-		if (DmlDataType == DML_TENSOR_DATA_TYPE_UNKNOWN)
-		{
-			return false;
-		}
-		
-		ElemSizeInBytes = Tensor.GetElemByteSize();
+	case ENNETensorDataType::Int32:
+		return DML_TENSOR_DATA_TYPE_INT32;
 
-		SetShape1D(Tensor.GetShape().GetData()[0], Rank);
-		Update(DmlDataType, Tensor.HasPreparedData());
+	case ENNETensorDataType::Int16:
+		return DML_TENSOR_DATA_TYPE_INT16;
 
-		return true;
+	case ENNETensorDataType::Int8:
+		return DML_TENSOR_DATA_TYPE_INT8;
+
+	default:
+		return DML_TENSOR_DATA_TYPE_UNKNOWN;
 	}
+}
 
-	void FTensorDesc::UpdateShapeAndStrides(TConstArrayView<uint32> InShape, TConstArrayView<uint32> InStrides)
-	{
-		check(!InShape.IsEmpty());
+} // namespace Util 
 
-		Sizes = InShape;
-		BuffDesc.Sizes = Sizes.GetData();
-		BuffDesc.DimensionCount = Sizes.Num();
+FTensorDescDml::FTensorDescDml()
+	: MinTensorRank(1)
+	, MaxTensorRank(DML_TENSOR_DIMENSION_COUNT_MAX1)
+	, bIsValidated(false)
+{
+	Sizes.Reset();
+	Strides.Reset();
+	BuffDesc = DML_BUFFER_TENSOR_DESC{};
+	Desc = DML_TENSOR_DESC{};
+}
 
-		if (!InStrides.IsEmpty())
-		{
-			check(InStrides.Num() == InShape.Num());
-			Strides = InStrides;
-			BuffDesc.Strides = Strides.GetData();
-		}
-		else
-		{
-			Strides.Reset();
-			BuffDesc.Strides = nullptr;
-		}
-	}
+FTensorDescDml& FTensorDescDml::SetTensorRank(int32 MinRank, int32 MaxRank)
+{
+	MinTensorRank = MinRank;
+	MaxTensorRank = MaxRank;
+	return *this;
+}
 
-	void FTensorDesc::SetStridesFromTensor(const NNECore::Internal::FTensor& InputDesc)
+FTensorDescDml& FTensorDescDml::SetShape(TConstArrayView<uint32> Shape)
+{
+	check(!bIsValidated);
+	
+	if (!bIsValidated)
 	{
-		uint32 CurrStride = 1;
-
-		Strides.SetNum(InputDesc.GetShape().Rank());
+		Sizes.Reset();
 		
-		for (int32 i = InputDesc.GetShape().Rank() - 1; i >= 0; --i)
+		if (Shape.Num() > 0)
 		{
-			Strides[i] = CurrStride;
-			CurrStride *= InputDesc.GetShape().GetData()[i];
+			Sizes.Append(Shape.GetData(), Shape.Num());
 		}
-	}
-
-	void FTensorDesc::Reset()
-	{
-		BuffDesc = DML_BUFFER_TENSOR_DESC{};
-		Desc = DML_TENSOR_DESC{};
-		Sizes.Reset();
-		Strides.Reset();
-		ElemSizeInBytes = 0;
-	}
-
-	void FTensorDesc::SetShape(TConstArrayView<uint32> Shape, int32 MinTensorRank)
-	{
-		Sizes = FSmallUIntArray(Shape.GetData(), Shape.Num());
-		const int32 DimOffset = MinTensorRank - Sizes.Num();
-
-		for (int Dim = 0; Dim < DimOffset; ++Dim)
+		// Handle scalar tensors
+		else if (Shape.Num() == 0)
 		{
-			Sizes.Insert(1, 0);
+			Sizes.Add(1);
 		}
-
 	}
 
-	void FTensorDesc::SetShapeAndStrides(TConstArrayView<uint32> Shape, TConstArrayView<uint32> BroadcastShape)
+	return *this;
+}
+
+FTensorDescDml& FTensorDescDml::SetShape(TConstArrayView<uint32> Shape, TConstArrayView<uint32> BroadcastShape)
+{
+	check(!bIsValidated);
+	
+	if (!bIsValidated)
 	{
 		const uint32 TargetDimension = BroadcastShape.Num() != -1 ? BroadcastShape.Num() : Shape.Num();
 		checkf(BroadcastShape.Num() >= Shape.Num(), TEXT("Can't broadcast tensor from rank %d to rank %d, should be inferior or equal."), Shape.Num(), TargetDimension);
-		
+
 		Sizes.SetNum(TargetDimension);
 		Strides.SetNum(TargetDimension);
 
 		const int32 DimensionOffset = int32(TargetDimension - Shape.Num());
-		
-		for (int32 i = 0; i < (int32) TargetDimension; ++i)
+
+		for (int32 i = 0; i < (int32)TargetDimension; ++i)
 		{
 			Sizes[i] = i < DimensionOffset ? 1 : Shape[i - DimensionOffset];
 		}
@@ -150,10 +145,23 @@ namespace DmlUtil
 		}
 	}
 
-	void FTensorDesc::SetShape1D(uint32 Dimension, int32 Rank)
+	return *this;
+}
+
+FTensorDescDml& FTensorDescDml::SetShape1D(uint32 Dimension, int32 Rank)
+{
+	check(!bIsValidated);
+	
+	if (!bIsValidated)
 	{
 		Sizes.Reset();
 
+		// Handle scalar tensors
+		if (Dimension == 0)
+		{
+			Dimension = 1;
+		}
+
 		Sizes.Add(1);
 		Sizes.Add(Dimension);
 		Sizes.Add(1);
@@ -164,185 +172,160 @@ namespace DmlUtil
 		}
 	}
 
-	void FTensorDesc::Update(DML_TENSOR_DATA_TYPE DataType, bool bHasWeightData)
-	{
-		BuffDesc.DataType = DataType;
-		BuffDesc.Flags = bHasWeightData ? DML_TENSOR_FLAG_OWNED_BY_DML : DML_TENSOR_FLAG_NONE;
-		BuffDesc.DimensionCount = Sizes.Num();
-		BuffDesc.Sizes = Sizes.GetData();
-		BuffDesc.Strides = Strides.IsEmpty() ? nullptr : Strides.GetData();
-		BuffDesc.TotalTensorSizeInBytes = CalculateBufferSize();
-
-		Desc = DML_TENSOR_DESC{ DML_TENSOR_TYPE_BUFFER, &BuffDesc };
-	}
+	return *this;
+}
 
-	uint64 FTensorDesc::CalculateBufferSize()
+FTensorDescDml& FTensorDescDml::SetStridesFromShape(TConstArrayView<uint32> Shape)
+{
+	check(!bIsValidated);
+	
+	if (!bIsValidated)
 	{
-		if (ElemSizeInBytes == 0)
+		uint32 CurrStride = 1;
+		Strides.SetNum(Shape.Num());
+
+		for (int32 i = Shape.Num() - 1; i >= 0; --i)
 		{
-			return 0;
+			Strides[i] = CurrStride;
+			CurrStride *= Shape[i];
 		}
+	}
 
-		uint64 MinSizeInBytes = 0;
-		
-		if (Strides.IsEmpty())
-		{
-			MinSizeInBytes = Sizes[0];
+	return *this;
+}
 
-			for (int32 i = 1; i < Sizes.Num(); ++i)
-			{
-				MinSizeInBytes *= Sizes[i];
-			}
+FTensorDescDml& FTensorDescDml::SetStrides(TConstArrayView<uint32> InStrides)
+{
+	check(!bIsValidated);
 
-			MinSizeInBytes *= ElemSizeInBytes;
-		}
-		else
-		{
-			uint32 IndexOfLastElement = 0;
+	if (!bIsValidated)
+	{
+		Strides = InStrides;
+	}
 
-			for (int32 i = 0; i < Sizes.Num(); ++i)
-			{
-				IndexOfLastElement += (Sizes[i] - 1) * Strides[i];
-			}
+	return *this;
+}
 
-			MinSizeInBytes = (static_cast<uint64>(IndexOfLastElement) + 1) * ElemSizeInBytes;
-		}
+FTensorDescDml& FTensorDescDml::SetDataType(ENNETensorDataType DataType)
+{
+	check(!bIsValidated);
+	
+	if (!bIsValidated)
+	{
+		DML_TENSOR_DATA_TYPE DmlDataType = Util::GetTensorDataType(DataType);
+		check(DmlDataType != DML_TENSOR_DATA_TYPE_UNKNOWN);
+	
+		uint64 ElemByteSize = CalculateBufferSize(UE::NNECore::GetTensorDataTypeSizeInBytes(DataType));
+		check(ElemByteSize > 0)
+
+		BuffDesc.DataType = DmlDataType;
+		BuffDesc.TotalTensorSizeInBytes = CalculateBufferSize(ElemByteSize);
+	}
 
-		// Round up to the nearest 4 bytes
-		MinSizeInBytes = (MinSizeInBytes + 3) & ~3ull;
+	return *this;
+}
 
-		return MinSizeInBytes;
+FTensorDescDml& FTensorDescDml::SetDataOwnedByDml(bool bSetOwnedByDml)
+{
+	check(!bIsValidated);
+	
+	if (!bIsValidated)
+	{
+		BuffDesc.Flags = bSetOwnedByDml ? DML_TENSOR_FLAG_OWNED_BY_DML : DML_TENSOR_FLAG_NONE;
 	}
 
+	return *this;
+}
 
-
-	void SetTensorStrides(FTensorDesc& TensorDesc, const NNECore::Internal::FTensor& InputDesc)
+bool FTensorDescDml::Validate()
+{
+	check(!bIsValidated);
+	
+	if (!bIsValidated)
 	{
-		uint32 CurrStride = 1;
-
-		TensorDesc.Strides.SetNum(InputDesc.GetShape().Rank());
-		
-		for (int32 i = InputDesc.GetShape().Rank() - 1; i >= 0; --i)
+		if (BuffDesc.DataType == DML_TENSOR_DATA_TYPE_UNKNOWN)
 		{
-			TensorDesc.Strides[i] = CurrStride;
-			CurrStride *= InputDesc.GetShape().GetData()[i];
+			UE_LOG(LogNNE, Error, TEXT("Invalid DML tensor data type"));
+			return false;
 		}
-	}
 
-	void SetTensorSizesAndStridesForBroadcast(FTensorDesc& TensorDesc, const NNECore::Internal::FTensor& InputDesc, const NNECore::Internal::FTensor& TargetDesc)
-	{
-		static_assert(NNECore::FTensorShape::MaxRank <= 8);
-		
-		const uint32 TargetDimension = TargetDesc.GetShape().Rank() != -1 ? TargetDesc.GetShape().Rank() : InputDesc.GetShape().Rank();
-		checkf(TargetDesc.GetShape().Rank() >= InputDesc.GetShape().Rank(), TEXT("Can't broadcast tensor from rank %d to rank %d, should be inferior or equal."), InputDesc.GetShape().Rank(), TargetDimension);
-		
-		TensorDesc.Sizes.SetNum(TargetDimension);
-		TensorDesc.Strides.SetNum(TargetDimension);
-
-		const int32 DimensionOffset = int32(TargetDimension - InputDesc.GetShape().Rank());
-		
-		for (int32 i = 0; i < (int32) TargetDimension; ++i)
+		if (BuffDesc.TotalTensorSizeInBytes == 0)
 		{
-			TensorDesc.Sizes[i] = i < DimensionOffset ? 1 : InputDesc.GetShape().GetData()[i - DimensionOffset];
+			UE_LOG(LogNNE, Error, TEXT("Invalid DML tensor size, it's 0"));
+			return false;
 		}
 
-		uint32 CurrStride = 1;
+		// Match minimum tensor rank
+		const int32 DimOffset = MinTensorRank - Sizes.Num();
 
-		for (int32 i = TargetDimension - 1; i >= 0; --i)
+		for (int Dim = 0; Dim < DimOffset; ++Dim)
 		{
-			const bool bBroadcast = TensorDesc.Sizes[i] < TargetDesc.GetShape().GetData()[i];
-
-			TensorDesc.Strides[i] = bBroadcast ? 0 : CurrStride;
-			CurrStride *= TensorDesc.Sizes[i];
-
-			TensorDesc.Sizes[i] = TargetDesc.GetShape().GetData()[i];
+			Sizes.Insert(1, 0);
 		}
-	}
 
-	bool IsSameShape(const NNECore::Internal::FTensor& Left, const NNECore::Internal::FTensor& Right)
-	{
-		if (Left.GetShape().Rank() != Right.GetShape().Rank())
+		const int32 Rank = Sizes.Num();
+
+		if (Rank < MinTensorRank || Rank > MaxTensorRank)
 		{
+			UE_LOG(LogNNE, Error, TEXT("Invalid DML tensor rank:%d [%d,%d]"), Rank, MinTensorRank, MaxTensorRank);
 			return false;
 		}
-		
-		for (int32 Idx = 0; Idx < Left.GetShape().Rank(); ++Idx)
-		{
-			if (Left.GetShape().GetData()[Idx] != Right.GetShape().GetData()[Idx])
-			{
-				return false;
-			}
-		}
-
-		return true;
-	}
 
-	DML_TENSOR_DATA_TYPE GetTensorDataType(ENNETensorDataType DataType)
-	{
-		switch (DataType)
+		if (!Strides.IsEmpty() && Strides.Num() != Rank)
 		{
-		case ENNETensorDataType::Double:
-			return DML_TENSOR_DATA_TYPE_FLOAT64;
-
-		case ENNETensorDataType::Float:
-			return DML_TENSOR_DATA_TYPE_FLOAT32;
-
-		case ENNETensorDataType::Half:
-			return DML_TENSOR_DATA_TYPE_FLOAT16;
-
-		case ENNETensorDataType::UInt64:
-			return DML_TENSOR_DATA_TYPE_UINT64;
-
-		case ENNETensorDataType::UInt32:
-			return DML_TENSOR_DATA_TYPE_UINT32;
-
-		case ENNETensorDataType::UInt16:
-			return DML_TENSOR_DATA_TYPE_UINT16;
-
-		case ENNETensorDataType::UInt8:
-			return DML_TENSOR_DATA_TYPE_UINT8;
-
-		case ENNETensorDataType::Int64:
-			return DML_TENSOR_DATA_TYPE_INT64;
+			UE_LOG(LogNNE, Error, TEXT("Invalid DML tensor stride rank:%d [size:%d]"), Rank, Strides.Num());
+			return false;
+		}
 
-		case ENNETensorDataType::Int32:
-			return DML_TENSOR_DATA_TYPE_INT32;
+		BuffDesc.DimensionCount = Sizes.Num();
+		BuffDesc.Sizes = Sizes.GetData();
+		BuffDesc.Strides = Strides.IsEmpty() ? nullptr : Strides.GetData();
+	
+		Desc.Type = DML_TENSOR_TYPE_BUFFER;
+		Desc.Desc = &BuffDesc;
 
-		case ENNETensorDataType::Int16:
-			return DML_TENSOR_DATA_TYPE_INT16;
+		bIsValidated = true;
+	}
 
-		case ENNETensorDataType::Int8:
-			return DML_TENSOR_DATA_TYPE_INT8;
+	return bIsValidated;
+}
 
-		default:
-			return DML_TENSOR_DATA_TYPE_UNKNOWN;
-		}
+uint64 FTensorDescDml::CalculateBufferSize(uint64 ElemSizeInBytes)
+{
+	if (ElemSizeInBytes == 0)
+	{
+		return 0;
 	}
 
-	uint64 CalculateBufferSize(const FTensorDesc& DmlTensor, const NNECore::Internal::FTensor& Desc)
-	{
-		uint64 ElemSizeInBytes = Desc.GetElemByteSize();
+	uint64 MinSizeInBytes = 0;
 		
-		if (ElemSizeInBytes == 0)
+	if (Strides.IsEmpty())
+	{
+		MinSizeInBytes = Sizes[0];
+
+		for (int32 i = 1; i < Sizes.Num(); ++i)
 		{
-			return 0;
+			MinSizeInBytes *= Sizes[i];
 		}
 
-		uint64 MinSizeInBytes = 0;		
+		MinSizeInBytes *= ElemSizeInBytes;
+	}
+	else
+	{
 		uint32 IndexOfLastElement = 0;
 
-		for (int32 i = 0; i < DmlTensor.Sizes.Num(); ++i)
+		for (int32 i = 0; i < Sizes.Num(); ++i)
 		{
-			IndexOfLastElement += (DmlTensor.Sizes[i] - 1) * DmlTensor.Strides[i];
+			IndexOfLastElement += (Sizes[i] - 1) * Strides[i];
 		}
 
 		MinSizeInBytes = (static_cast<uint64>(IndexOfLastElement) + 1) * ElemSizeInBytes;
-		
-		// Round up to the nearest 4 bytes.
-		MinSizeInBytes = (MinSizeInBytes + 3) & ~3ull;
-
-		return MinSizeInBytes;
 	}
+
+	// Round up to the nearest 4 bytes
+	MinSizeInBytes = (MinSizeInBytes + 3) & ~3ull;
+
+	return MinSizeInBytes;
 }
 
 //
@@ -363,104 +346,6 @@ IDMLOperator* FOperatorDml::GetOperator()
 	return DmlOp;
 }
 
-bool FOperatorDml::InitDmlTensorDesc(DmlUtil::FTensorDesc& DmlTensorDesc, const NNECore::Internal::FTensor& Tensor)
-{
-	DML_TENSOR_DATA_TYPE DmlDataType = DmlUtil::GetTensorDataType(Tensor.GetDataType());
-
-	if (DmlDataType == DML_TENSOR_DATA_TYPE_UNKNOWN)
-	{
-		DmlTensorDesc.BuffDesc = DML_BUFFER_TENSOR_DESC{};
-		DmlTensorDesc.Desc = DML_TENSOR_DESC{};
-
-		return false;
-	}
-
-	DmlTensorDesc.Sizes = Tensor.GetShape().GetData();
-	//Note: We should support tensor padding using strides defined in FTensorDesc
-	//DmlUtil::SetTensorStrides(DmlTensorDesc, TensorDesc.Strides);
-		
-	DML_BUFFER_TENSOR_DESC& BuffDesc = DmlTensorDesc.BuffDesc;
-
-	BuffDesc = DML_BUFFER_TENSOR_DESC{};
-	BuffDesc.DataType = DmlDataType;
-	BuffDesc.Flags = Tensor.HasPreparedData() ? DML_TENSOR_FLAG_OWNED_BY_DML : DML_TENSOR_FLAG_NONE;
-	BuffDesc.DimensionCount = Tensor.GetShape().Rank();
-	BuffDesc.Sizes = DmlTensorDesc.Sizes.GetData();
-	BuffDesc.Strides = nullptr;
-	BuffDesc.TotalTensorSizeInBytes = Tensor.GetDataSize(); // DmlUtil::CalculateBufferSize(DmlTensorDesc, Tensor);
-
-	static DmlUtil::FSmallUIntArray ScalarShape({ 1 });
-
-	//Handle scalar tensors
-	if (Tensor.GetShape().Rank() == 0)
-	{
-		BuffDesc.DimensionCount = ScalarShape.Num();
-		DmlTensorDesc.Sizes = ScalarShape;
-		BuffDesc.Sizes = ScalarShape.GetData();
-	}
-
-	DmlTensorDesc.Desc = DML_TENSOR_DESC{ DML_TENSOR_TYPE_BUFFER, &DmlTensorDesc.BuffDesc };
-
-	return true;
-}
-
-bool FOperatorDml::InitDmlTensorDesc(DmlUtil::FTensorDesc& DmlTensorDesc, const NNECore::Internal::FTensor& Tensor, const NNECore::Internal::FTensor& Broadcast)
-{
-	DML_TENSOR_DATA_TYPE DmlDataType = DmlUtil::GetTensorDataType(Tensor.GetDataType());
-
-	if (DmlDataType == DML_TENSOR_DATA_TYPE_UNKNOWN)
-	{
-		DmlTensorDesc.BuffDesc = DML_BUFFER_TENSOR_DESC{};
-		DmlTensorDesc.Desc = DML_TENSOR_DESC{};
-
-		return false;
-	}
-
-	if (DmlUtil::IsSameShape(Tensor, Broadcast))
-	{
-		DmlTensorDesc.Sizes = Tensor.GetShape().GetData();
-		DmlUtil::SetTensorStrides(DmlTensorDesc, Tensor);
-	}
-	else if (Tensor.GetShape().Rank() > Broadcast.GetShape().Rank())
-	{
-		return false;
-	}
-	else
-	{
-		DmlUtil::SetTensorSizesAndStridesForBroadcast(DmlTensorDesc, Tensor, Broadcast);
-	}
-
-	//UE_LOG(LogNNE, Warning, TEXT("DmlTensorDesc:%d,%d,%d -> %d,%d,%d"),
-	//	Tensor.Sizes[0],
-	//	Tensor.Dimension > 1 ? Tensor.Sizes[1] : 0,
-	//	Tensor.Dimension > 2 ? Tensor.Sizes[2] : 0,
-	//	DmlTensor.Sizes[0],
-	//	DmlTensor.Dimension > 1 ? DmlTensor.Sizes[1] : 0,
-	//	DmlTensor.Dimension > 2 ? DmlTensor.Sizes[2] : 0
-	//);
-
-	//UE_LOG(LogNNE, Warning, TEXT("DmlTensorStrides:%d,%d,%d"),
-	//	DmlTensorDesc.Strides[0], DmlTensorDesc.Strides[1], DmlTensorDesc.Strides[2]
-	//);
-
-	check(DmlTensorDesc.Strides.Num() == DmlTensorDesc.Sizes.Num());
-		
-	DML_BUFFER_TENSOR_DESC& BuffDesc = DmlTensorDesc.BuffDesc;
-		
-	BuffDesc = DML_BUFFER_TENSOR_DESC{};
-
-	BuffDesc.DataType = DmlDataType;
-	BuffDesc.Flags = Tensor.HasPreparedData() ? DML_TENSOR_FLAG_OWNED_BY_DML : DML_TENSOR_FLAG_NONE;
-	BuffDesc.DimensionCount = DmlTensorDesc.Sizes.Num();
-	BuffDesc.Sizes = DmlTensorDesc.Sizes.GetData();
-	BuffDesc.Strides = DmlTensorDesc.Strides.GetData();
-	BuffDesc.TotalTensorSizeInBytes = DmlUtil::CalculateBufferSize(DmlTensorDesc, Tensor);
-		
-	DmlTensorDesc.Desc = DML_TENSOR_DESC{ DML_TENSOR_TYPE_BUFFER, &DmlTensorDesc.BuffDesc };
-
-	return true;
-}
-
 // Create DirectML operator	
 bool FOperatorDml::CreateOperator(IDMLDevice* Device, const DML_OPERATOR_DESC& DmlOpDesc)
 {
diff --git a/Engine/Plugins/Experimental/NNE/Source/NNERuntimeRDG/Private/Dml/NNEDmlOperator.h b/Engine/Plugins/Experimental/NNE/Source/NNERuntimeRDG/Private/Dml/NNEDmlOperator.h
--- a/Engine/Plugins/Experimental/NNE/Source/NNERuntimeRDG/Private/Dml/NNEDmlOperator.h
+++ b/Engine/Plugins/Experimental/NNE/Source/NNERuntimeRDG/Private/Dml/NNEDmlOperator.h
@@ -54,119 +54,217 @@ public:
 	TComPtr<IDMLCommandRecorder>	CmdRec{ nullptr };
 };
 
-namespace DmlUtil
+class FTensorDescDml;
+
+namespace Util
 {
-	template<typename T>
-	using FSmallArray = TArray<T, TInlineAllocator<NNECore::FTensorShape::MaxRank>>;
-	using FSmallIntArray = TArray<int32, TInlineAllocator<NNECore::FTensorShape::MaxRank>>;
-	using FSmallUIntArray = TArray<uint32, TInlineAllocator<NNECore::FTensorShape::MaxRank>>;
+	
+template<typename T>
+using FSmallArray = TArray<T, TInlineAllocator<NNECore::FTensorShape::MaxRank>>;
+using FSmallIntArray = TArray<int32, TInlineAllocator<NNECore::FTensorShape::MaxRank>>;
+using FSmallUIntArray = TArray<uint32, TInlineAllocator<NNECore::FTensorShape::MaxRank>>;
 
-	template<typename InputType, typename OutputType>
-	inline bool IsOverflowing(InputType Input)
+template<typename InputType, typename OutputType>
+inline bool IsOverflowing(InputType Input)
+{
+	OutputType Output = static_cast<OutputType>(Input);
+	if(Input != static_cast<InputType>(Output))
 	{
-		OutputType Output = static_cast<OutputType>(Input);
-		if(Input != static_cast<InputType>(Output))
-		{
-			return true;
-		}
-		return false;
+		return true;
 	}
+	return false;
+}
 
-	template<typename InputType, typename OutputType>
-	inline bool ConvertArrayViewNoOverflow(TConstArrayView<InputType> InputView, TArrayView<OutputType>& OutputView)
+template<typename InputType, typename OutputType>
+inline bool ConvertArrayViewNoOverflow(TConstArrayView<InputType> InputView, TArrayView<OutputType>& OutputView)
+{
+	OutputView = MakeArrayView((OutputType*) InputView.GetData(), InputView.Num());
+	for(int32 Idx = 0; Idx < InputView.Num(); ++Idx)
 	{
-		OutputView = MakeArrayView((OutputType*) InputView.GetData(), InputView.Num());
-		for(int32 Idx = 0; Idx < InputView.Num(); ++Idx)
+		if(InputView[Idx] != static_cast<InputType>(OutputView[Idx]))
 		{
-			if(InputView[Idx] != static_cast<InputType>(OutputView[Idx]))
-			{
-				return false;
-			}
+			return false;
 		}
-		return true;
 	}
+	return true;
+}
 
-	template<typename OutputType, typename AllocatorType>
-	inline bool GetArrayAttributeNoOverflow(
-		const FNNEAttributeValue* Attr, 
-		TArray<OutputType, AllocatorType>& OutputArray, 
-		TConstArrayView<OutputType> DefaultValues = MakeEmptyConstArrayView<OutputType>()
-		)
+template<typename OutputType, typename AllocatorType>
+inline bool GetArrayAttributeNoOverflow(
+	const FNNEAttributeValue* Attr, 
+	TArray<OutputType, AllocatorType>& OutputArray, 
+	TConstArrayView<OutputType> DefaultValues = MakeEmptyConstArrayView<OutputType>()
+	)
+{
+	if (Attr)
 	{
-		if (Attr)
-		{
 
-			TArrayView<OutputType> ConvertedView;
-			TArray<int32> IntArray;
-			TArray<float> FloatArray;
+		TArrayView<OutputType> ConvertedView;
+		TArray<int32> IntArray;
+		TArray<float> FloatArray;
 
-			switch(Attr->GetType())
+		switch(Attr->GetType())
+		{
+		case ENNEAttributeDataType::Int32Array:
 			{
-			case ENNEAttributeDataType::Int32Array:
+				IntArray = Attr->GetValue<TArray<int32>>();
+				if(!ConvertArrayViewNoOverflow(TConstArrayView<int32>(IntArray), ConvertedView))
 				{
-					IntArray = Attr->GetValue<TArray<int32>>();
-					if(!ConvertArrayViewNoOverflow(TConstArrayView<int32>(IntArray), ConvertedView))
-					{
-						return false;
-					}
+					return false;
 				}
-				break;
-			case ENNEAttributeDataType::FloatArray:
+			}
+			break;
+		case ENNEAttributeDataType::FloatArray:
+			{
+				FloatArray = Attr->GetValue<TArray<float>>();
+				if(!ConvertArrayViewNoOverflow(TConstArrayView<float>(FloatArray), ConvertedView))
 				{
-					FloatArray = Attr->GetValue<TArray<float>>();
-					if(!ConvertArrayViewNoOverflow(TConstArrayView<float>(FloatArray), ConvertedView))
-					{
-						return false;
-					}
+					return false;
 				}
-				break;
-			default:
-				return false;
 			}
-			
-			OutputArray = TArray<OutputType, AllocatorType>{ ConvertedView };
+			break;
+		default:
+			return false;
 		}
-		else
-		{
-			OutputArray = TArray<OutputType, AllocatorType>{ DefaultValues };
-		}
-		return true;
+			
+		OutputArray = TArray<OutputType, AllocatorType>{ ConvertedView };
 	}
-
-	struct FTensorDesc
+	else
 	{
-		DML_BUFFER_TENSOR_DESC	BuffDesc;
-		DML_TENSOR_DESC			Desc;
-		// Don't edit Sizes and Strides directly, use methods
-		FSmallUIntArray			Sizes;
-		FSmallUIntArray			Strides;
-		uint64					ElemSizeInBytes;
-
-		bool InitFromTensor(const NNECore::Internal::FTensor& InputDesc, int32 MinTensorRank, TConstArrayView<uint32> Broadcast = MakeArrayView((uint32*) nullptr, 0), TConstArrayView<uint32> CustomShape = MakeArrayView((uint32*) nullptr, 0));
-		bool InitFromTensor1D(const NNECore::Internal::FTensor& InputDesc, int32 Rank);
-
-		void SetStridesFromTensor(const NNECore::Internal::FTensor& InputDesc);
-
-		void UpdateShapeAndStrides(TConstArrayView<uint32> InShape, TConstArrayView<uint32> InStrides = MakeArrayView((uint32*) nullptr, 0));
-
-	private:
-		void Reset();
-
-		void SetShape(TConstArrayView<uint32> Shape, int32 MinTensorRank);
-		void SetShapeAndStrides(TConstArrayView<uint32> Shape, TConstArrayView<uint32> BroadcastShape);
-		void SetShape1D(uint32 Dimension, int32 Rank);
-		
-		void Update(DML_TENSOR_DATA_TYPE DataType, bool bHasWeightData = false);
-		uint64 CalculateBufferSize();
-	};
-
-	extern void SetTensorStrides(FTensorDesc& TensorDesc, const NNECore::Internal::FTensor& InputDesc);
-	extern void SetTensorSizesAndStridesForBroadcast(FTensorDesc& TensorDesc, const NNECore::Internal::FTensor& InputDesc, const NNECore::Internal::FTensor& TargetDesc);
-	extern bool IsSameShape(const NNECore::Internal::FTensor& Left, const NNECore::Internal::FTensor& Right);
-	DML_TENSOR_DATA_TYPE GetTensorDataType(ENNETensorDataType DataType);
-	uint64 CalculateBufferSize(const FTensorDesc& DmlTensor, const NNECore::Internal::FTensor& Desc);
+		OutputArray = TArray<OutputType, AllocatorType>{ DefaultValues };
+	}
+	return true;
 }
 
+extern bool IsSameShape(const NNECore::Internal::FTensor& Left, const NNECore::Internal::FTensor& Right);
+extern bool IsSameShape(const FTensorDescDml& Left, const FTensorDescDml& Right);
+extern DML_TENSOR_DATA_TYPE GetTensorDataType(ENNETensorDataType DataType);
+
+} // Util
+
+class FTensorDescDml
+{
+public:
+
+	FTensorDescDml();
+
+	// Set DML operator tensor minimum and maximum rank.
+	// By default minimum and maximum range is between 1 and DML_TENSOR_DIMENSION_COUNT_MAX1 (8).
+	FTensorDescDml& SetTensorRank(int32 MinRank, int32 MaxRank);
+
+	// Set tensor shape
+	FTensorDescDml& SetShape(TConstArrayView<uint32> Shape);
+
+	// Utility method that calls SetShape() from above
+	FTensorDescDml& SetShape(const UE::NNECore::FTensorShape& Shape)
+	{
+		return SetShape(Shape.GetData());
+	}
+
+	// Set shape and match the broadcast shape
+	FTensorDescDml& SetShape(TConstArrayView<uint32> Shape, TConstArrayView<uint32> BroadcastShape);
+
+	// Utility method that calles SetShape() from above
+	FTensorDescDml& SetShape(const UE::NNECore::FTensorShape& Shape, const UE::NNECore::FTensorShape& BroadcastShape)
+	{
+		return SetShape(Shape.GetData(), BroadcastShape.GetData());
+	}
+
+	// Use for tensors that have shape [1, C, 1] or [1, C, 1, 1] etc.
+	FTensorDescDml& SetShape1D(uint32 Dimension, int32 Rank);
+	
+	// Set strides computed from shape (used for broadcasting)
+	FTensorDescDml& SetStridesFromShape(TConstArrayView<uint32> Shape);
+
+	// Set strides computed from shape (used for broadcasting)
+	FTensorDescDml& SetStridesFromShape(const UE::NNECore::FTensorShape& Shape)
+	{
+		return SetStridesFromShape(Shape.GetData());
+	}
+
+	// Set computed strides
+	FTensorDescDml& SetStrides(TConstArrayView<uint32> Strides);
+
+	// Set tensor data type
+	FTensorDescDml& SetDataType(ENNETensorDataType DataType);
+
+	// Set tensor flags
+	FTensorDescDml& SetDataOwnedByDml(bool bSetOwnedByDml);
+
+	// Utility method to use FTensor to set all the members
+	FTensorDescDml& SetFromTensor(const UE::NNECore::Internal::FTensor& Tensor)
+	{
+		return SetShape(Tensor.GetShape())
+			.SetDataType(Tensor.GetDataType())
+			.SetDataOwnedByDml(Tensor.HasPreparedData());
+	}
+
+	// Utility method to use FTensor to set all the members
+	FTensorDescDml& SetFromTensorBroadcast(const UE::NNECore::Internal::FTensor& Tensor, const UE::NNECore::FTensorShape& BroadcastShape)
+	{
+		return SetShape(Tensor.GetShape(), BroadcastShape)
+			.SetDataType(Tensor.GetDataType())
+			.SetDataOwnedByDml(Tensor.HasPreparedData());
+	}
+
+	// Utility method to use FTensor to set all the members
+	FTensorDescDml& SetFromTensor1D(const UE::NNECore::Internal::FTensor& Tensor, int32 Rank)
+	{
+		return SetShape1D(Tensor.GetShape().Rank(), Rank)
+			.SetDataType(Tensor.GetDataType())
+			.SetDataOwnedByDml(Tensor.HasPreparedData());
+	}
+
+	// Return computed sizes
+	TConstArrayView<uint32> GetSizes() const
+	{
+		return Sizes;
+	}
+
+	// Return computed strides
+	TConstArrayView<uint32> GetStrides() const
+	{
+		return Strides;
+	}
+
+	// Return computed size rank
+	int32 GetRank() const
+	{
+		return Sizes.Num();
+	}
+
+	// Validate the tensor descriptor, once it's validated any calls to SetXXX() will be ignored
+	bool Validate();
+
+	// Return filled DML tensor descriptor
+	// NOTE: Call this method only after Validate() is called
+	const DML_TENSOR_DESC* GetDmlDesc() const
+	{
+		check(bIsValidated);
+		return &Desc;
+	}
+
+	// Return DML tensor data type
+	// NOTE: Call this method only after Validate() is called
+	DML_TENSOR_DATA_TYPE GetDmlDataType() const
+	{
+		check(bIsValidated);
+		return BuffDesc.DataType;
+	}
+
+private:
+
+	uint64 CalculateBufferSize(uint64 ElemSizeInBytes);
+
+	Util::FSmallUIntArray	Sizes;
+	Util::FSmallUIntArray	Strides;
+	DML_BUFFER_TENSOR_DESC	BuffDesc;
+	DML_TENSOR_DESC			Desc;
+	int32					MinTensorRank;
+	int32					MaxTensorRank;
+	bool					bIsValidated;
+};
+
 //
 // DirectML operator base class
 //
@@ -185,15 +283,12 @@ public:
 	IDMLOperator* GetOperator();
 
 protected:
-
-	bool InitDmlTensorDesc(DmlUtil::FTensorDesc& DmlTensorDesc, const NNECore::Internal::FTensor& Tensor);
-	bool InitDmlTensorDesc(DmlUtil::FTensorDesc& DmlTensorDesc, const NNECore::Internal::FTensor& Tensor, const NNECore::Internal::FTensor& Broadcast);
 	
 	bool CreateOperator(IDMLDevice* Device, const DML_OPERATOR_DESC& DmlOpDesc);
 
 	TComPtr<IDMLOperator>		DmlOp;
-	DmlUtil::FSmallIntArray		ConstantCPUInputs;
-	DmlUtil::FSmallIntArray		RemappedInputs;
+	Util::FSmallIntArray		ConstantCPUInputs;
+	Util::FSmallIntArray		RemappedInputs;
 };
 
 /**
diff --git a/Engine/Plugins/Experimental/NNE/Source/NNERuntimeRDG/Private/Dml/NNEDmlOperatorActivations.cpp b/Engine/Plugins/Experimental/NNE/Source/NNERuntimeRDG/Private/Dml/NNEDmlOperatorActivations.cpp
--- a/Engine/Plugins/Experimental/NNE/Source/NNERuntimeRDG/Private/Dml/NNEDmlOperatorActivations.cpp
+++ b/Engine/Plugins/Experimental/NNE/Source/NNERuntimeRDG/Private/Dml/NNEDmlOperatorActivations.cpp
@@ -40,15 +40,15 @@ public:
 	//
 	virtual bool Initialize(IDMLDevice* Device, TArrayView<const NNECore::Internal::FTensor> InputTensors, TArrayView<const NNECore::Internal::FTensor> OutputTensors, const NNECore::FAttributeMap& Attributes) override
 	{
-		const NNECore::Internal::FTensor& InputTensorDesc = InputTensors[0];
-		const NNECore::Internal::FTensor& OutputTensorDesc = OutputTensors[0];
+		const NNECore::Internal::FTensor& InputTensor = InputTensors[0];
+		const NNECore::Internal::FTensor& OutputTensor = OutputTensors[0];
 
 		if constexpr (std::is_same_v<DmlActivationOpDescType, DML_ACTIVATION_SOFTMAX1_OPERATOR_DESC> ||
 					  std::is_same_v<DmlActivationOpDescType, DML_ACTIVATION_LOG_SOFTMAX1_OPERATOR_DESC>)
 		{
 			Axis = Attributes.GetValueOrDefault(TEXT("axis"), -1);
 		}
-		Axis = HandleNegativeAxis(Axis, InputTensorDesc.GetShape().Rank());
+		Axis = HandleNegativeAxis(Axis, InputTensor.GetShape().Rank());
 
 		Alpha = Attributes.GetValueOrDefault(TEXT("alpha"), Alpha);
 		Beta = Attributes.GetValueOrDefault(TEXT("beta"), Beta);
@@ -56,17 +56,19 @@ public:
 
 
 		// Initialize tensor descriptor (it's same for both input and output)
-		DmlUtil::FTensorDesc	DmlTensorDesc{};
+		FTensorDescDml	DmlTensorDesc;
 
-		if (!InitDmlTensorDesc(DmlTensorDesc, InputTensorDesc))
+		if (!DmlTensorDesc
+				.SetFromTensor(InputTensor)
+				.Validate())
 		{
-			UE_LOG(LogNNE, Warning, TEXT("Failed to initialize tensor(s) for DML inference"));
+			UE_LOG(LogNNE, Error, TEXT("Failed to initialize tensor(s) for DML inference"));
 			return false;
 		}
 
 		DmlActivationOpDescType	DmlActivationOpDesc{};
 
-		InitDmlOpDesc(DmlActivationOpDesc, DmlTensorDesc);
+		InitDmlOpDesc(DmlActivationOpDesc, *DmlTensorDesc.GetDmlDesc());
 
 		DML_OPERATOR_DESC DmlOpDesc{};
 
@@ -79,62 +81,62 @@ public:
 private:
 
 	template<typename OpDesc>
-	void InitDmlOpDesc(OpDesc& Desc, DmlUtil::FTensorDesc& TensorDesc)
+	void InitDmlOpDesc(OpDesc& Desc, const DML_TENSOR_DESC& TensorDesc)
 	{
-		Desc.InputTensor = &TensorDesc.Desc;
-		Desc.OutputTensor = &TensorDesc.Desc;
+		Desc.InputTensor = &TensorDesc;
+		Desc.OutputTensor = &TensorDesc;
 	}
 
-	void InitDmlOpDesc(DML_ACTIVATION_LOG_SOFTMAX1_OPERATOR_DESC& Desc, DmlUtil::FTensorDesc& TensorDesc)
+	void InitDmlOpDesc(DML_ACTIVATION_LOG_SOFTMAX1_OPERATOR_DESC& Desc, const DML_TENSOR_DESC& TensorDesc)
 	{
-		Desc.InputTensor = &TensorDesc.Desc;
-		Desc.OutputTensor = &TensorDesc.Desc;
+		Desc.InputTensor = &TensorDesc;
+		Desc.OutputTensor = &TensorDesc;
 		Desc.AxisCount = 1;
 		Desc.Axes = (UINT*) &Axis;
 	}
 
-	void InitDmlOpDesc(DML_ACTIVATION_SOFTMAX1_OPERATOR_DESC& Desc, DmlUtil::FTensorDesc& TensorDesc)
+	void InitDmlOpDesc(DML_ACTIVATION_SOFTMAX1_OPERATOR_DESC& Desc, const DML_TENSOR_DESC& TensorDesc)
 	{
-		Desc.InputTensor = &TensorDesc.Desc;
-		Desc.OutputTensor = &TensorDesc.Desc;
+		Desc.InputTensor = &TensorDesc;
+		Desc.OutputTensor = &TensorDesc;
 		Desc.AxisCount = 1;
 		Desc.Axes = (UINT*) &Axis;
 	}
 
-	void InitDmlOpDesc(DML_ACTIVATION_SOFTPLUS_OPERATOR_DESC& Desc, DmlUtil::FTensorDesc& TensorDesc)
+	void InitDmlOpDesc(DML_ACTIVATION_SOFTPLUS_OPERATOR_DESC& Desc, const DML_TENSOR_DESC& TensorDesc)
 	{
-		Desc.InputTensor = &TensorDesc.Desc;
-		Desc.OutputTensor = &TensorDesc.Desc;
+		Desc.InputTensor = &TensorDesc;
+		Desc.OutputTensor = &TensorDesc;
 		Desc.Steepness = 1.0f;
 	}
 
-	void InitDmlOpDesc(DML_ACTIVATION_SCALED_ELU_OPERATOR_DESC& Desc, DmlUtil::FTensorDesc& TensorDesc)
+	void InitDmlOpDesc(DML_ACTIVATION_SCALED_ELU_OPERATOR_DESC& Desc, const DML_TENSOR_DESC& TensorDesc)
 	{
-		Desc.InputTensor = &TensorDesc.Desc;
-		Desc.OutputTensor = &TensorDesc.Desc;
+		Desc.InputTensor = &TensorDesc;
+		Desc.OutputTensor = &TensorDesc;
 		Desc.Alpha = Alpha;
 		Desc.Gamma = Gamma;
 	}
 
-	void InitDmlOpDesc(DML_ACTIVATION_ELU_OPERATOR_DESC& Desc, DmlUtil::FTensorDesc& TensorDesc)
+	void InitDmlOpDesc(DML_ACTIVATION_ELU_OPERATOR_DESC& Desc, const DML_TENSOR_DESC& TensorDesc)
 	{
-		Desc.InputTensor = &TensorDesc.Desc;
-		Desc.OutputTensor = &TensorDesc.Desc;
+		Desc.InputTensor = &TensorDesc;
+		Desc.OutputTensor = &TensorDesc;
 		Desc.Alpha = Alpha;
 	}
 
-	void InitDmlOpDesc(DML_ACTIVATION_HARD_SIGMOID_OPERATOR_DESC& Desc, DmlUtil::FTensorDesc& TensorDesc)
+	void InitDmlOpDesc(DML_ACTIVATION_HARD_SIGMOID_OPERATOR_DESC& Desc, const DML_TENSOR_DESC& TensorDesc)
 	{
-		Desc.InputTensor = &TensorDesc.Desc;
-		Desc.OutputTensor = &TensorDesc.Desc;
+		Desc.InputTensor = &TensorDesc;
+		Desc.OutputTensor = &TensorDesc;
 		Desc.Alpha = Alpha;
 		Desc.Beta = Beta;
 	}
 
-	void InitDmlOpDesc(DML_ACTIVATION_LEAKY_RELU_OPERATOR_DESC& Desc, DmlUtil::FTensorDesc& TensorDesc)
+	void InitDmlOpDesc(DML_ACTIVATION_LEAKY_RELU_OPERATOR_DESC& Desc, const DML_TENSOR_DESC& TensorDesc)
 	{
-		Desc.InputTensor = &TensorDesc.Desc;
-		Desc.OutputTensor = &TensorDesc.Desc;
+		Desc.InputTensor = &TensorDesc;
+		Desc.OutputTensor = &TensorDesc;
 		Desc.Alpha = Alpha;
 	}
 };
@@ -168,36 +170,42 @@ public:
 	//
 	virtual bool Initialize(IDMLDevice* Device, TArrayView<const NNECore::Internal::FTensor> InputTensors, TArrayView<const NNECore::Internal::FTensor> OutputTensors, const NNECore::FAttributeMap& Attributes) override
 	{
-		const NNECore::Internal::FTensor& InputATensorDesc = InputTensors[0];
-		const NNECore::Internal::FTensor& InputBTensorDesc = InputTensors[1];
-		const NNECore::Internal::FTensor& OutputTensorDesc = OutputTensors[0];
+		const NNECore::Internal::FTensor& InputATensor = InputTensors[0];
+		const NNECore::Internal::FTensor& InputBTensor = InputTensors[1];
+		const NNECore::Internal::FTensor& OutputTensor = OutputTensors[0];
 
 		// Initialize tensor descriptors
-		DmlUtil::FTensorDesc	DmlInputATensorDesc{};
-		DmlUtil::FTensorDesc	DmlInputBTensorDesc{};
-		DmlUtil::FTensorDesc	DmlOutputTensorDesc{};
+		FTensorDescDml	DmlInputATensorDesc;
+		FTensorDescDml	DmlInputBTensorDesc;
+		FTensorDescDml	DmlOutputTensorDesc;
 
-		if (!InitDmlTensorDesc(DmlInputATensorDesc, InputATensorDesc, OutputTensorDesc))
+		if (!DmlInputATensorDesc
+				.SetFromTensorBroadcast(InputATensor, OutputTensor.GetShape())
+				.Validate())
 		{
-			UE_LOG(LogNNE, Warning, TEXT("Failed to initialize tensor(s) for DML inference"));
+			UE_LOG(LogNNE, Error, TEXT("Failed to initialize tensor(s) for DML inference"));
 			return false;
 		}
 
-		if (!InitDmlTensorDesc(DmlInputBTensorDesc, InputBTensorDesc, OutputTensorDesc))
+		if (!DmlInputBTensorDesc
+				.SetFromTensorBroadcast(InputBTensor, OutputTensor.GetShape())
+				.Validate())
 		{
-			UE_LOG(LogNNE, Warning, TEXT("Failed to initialize tensor(s) for DML inference"));
+			UE_LOG(LogNNE, Error, TEXT("Failed to initialize tensor(s) for DML inference"));
 			return false;
 		}
 
-		if (!InitDmlTensorDesc(DmlOutputTensorDesc, OutputTensorDesc))
+		if (!DmlOutputTensorDesc
+				.SetFromTensor(OutputTensor)
+				.Validate())
 		{
-			UE_LOG(LogNNE, Warning, TEXT("Failed to initialize tensor(s) for DML inference"));
+			UE_LOG(LogNNE, Error, TEXT("Failed to initialize tensor(s) for DML inference"));
 			return false;
 		}
 
 		DmlActivationOpDescType	DmlActivationOpDesc{};
 
-		InitDmlOpDesc(DmlActivationOpDesc, DmlInputATensorDesc, DmlInputBTensorDesc, DmlOutputTensorDesc);
+		InitDmlOpDesc(DmlActivationOpDesc, *DmlInputATensorDesc.GetDmlDesc(), *DmlInputBTensorDesc.GetDmlDesc(), *DmlOutputTensorDesc.GetDmlDesc());
 
 		DML_OPERATOR_DESC DmlOpDesc{};
 
@@ -210,18 +218,18 @@ public:
 private:
 
 	template<typename OpDesc>
-	void InitDmlOpDesc(OpDesc& Desc, DmlUtil::FTensorDesc& LHSTensor, DmlUtil::FTensorDesc& RHSTensor, DmlUtil::FTensorDesc& OutputTensor)
+	void InitDmlOpDesc(OpDesc& Desc, DML_TENSOR_DESC& LHSTensorDesc, const DML_TENSOR_DESC& RHSTensorDesc, const DML_TENSOR_DESC& OutputTensorDesc)
 	{
-		Desc.ATensor = &LHSTensor.Desc;
-		Desc.BTensor = &RHSTensor.Desc;
-		Desc.OutputTensor = &OutputTensor.Desc;
+		Desc.ATensor = &LHSTensorDesc;
+		Desc.BTensor = &RHSTensorDesc;
+		Desc.OutputTensor = &OutputTensorDesc;
 	}
 
-	void InitDmlOpDesc(DML_ACTIVATION_PARAMETERIZED_RELU_OPERATOR_DESC& Desc, DmlUtil::FTensorDesc& LHSTensor, DmlUtil::FTensorDesc& RHSTensor, DmlUtil::FTensorDesc& OutputTensor)
+	void InitDmlOpDesc(DML_ACTIVATION_PARAMETERIZED_RELU_OPERATOR_DESC& Desc, const DML_TENSOR_DESC& LHSTensorDesc, const DML_TENSOR_DESC& RHSTensorDesc, const DML_TENSOR_DESC& OutputTensorDesc)
 	{
-		Desc.InputTensor = &LHSTensor.Desc;
-		Desc.SlopeTensor = &RHSTensor.Desc;
-		Desc.OutputTensor = &OutputTensor.Desc;
+		Desc.InputTensor = &LHSTensorDesc;
+		Desc.SlopeTensor = &RHSTensorDesc;
+		Desc.OutputTensor = &OutputTensorDesc;
 	}
 };
 
diff --git a/Engine/Plugins/Experimental/NNE/Source/NNERuntimeRDG/Private/Dml/NNEDmlOperatorBatchNormalization.cpp b/Engine/Plugins/Experimental/NNE/Source/NNERuntimeRDG/Private/Dml/NNEDmlOperatorBatchNormalization.cpp
--- a/Engine/Plugins/Experimental/NNE/Source/NNERuntimeRDG/Private/Dml/NNEDmlOperatorBatchNormalization.cpp
+++ b/Engine/Plugins/Experimental/NNE/Source/NNERuntimeRDG/Private/Dml/NNEDmlOperatorBatchNormalization.cpp
@@ -46,7 +46,7 @@ public:
 
 		if (InputTensor.GetShape().Rank() > 8)
 		{
-			UE_LOG(LogNNE, Warning, TEXT("InputTensor rank should be between 1 and 8, got:%d"), InputTensor.GetShape().Rank());
+			UE_LOG(LogNNE, Error, TEXT("InputTensor rank should be between 1 and 8, got:%d"), InputTensor.GetShape().Rank());
 			return false;
 		}
 
@@ -56,59 +56,65 @@ public:
 
 		if (bTrainingMode || OutputTensors.Num() > 1)
 		{
-			UE_LOG(LogNNE, Warning, TEXT("DML:BatchNormalization doesn't support training mode"));
+			UE_LOG(LogNNE, Error, TEXT("DML:BatchNormalization doesn't support training mode"));
 			return false;
 		}
 
-		DmlUtil::FTensorDesc	DmlInputTensors[Count];
+		FTensorDescDml DmlInputTensorDescs[Count];
 			
-		if (!DmlInputTensors[X].InitFromTensor(InputTensor, InputTensor.GetShape().Rank()))
+		if (!DmlInputTensorDescs[X]
+				.SetFromTensor(InputTensor)
+				.Validate())
 		{
-			UE_LOG(LogNNE, Warning, TEXT("Failed to initialize tensor(s) for DML inference"));
+			UE_LOG(LogNNE, Error, TEXT("Failed to initialize tensor(s) for DML inference"));
 			return false;
 		}
 
 		for (int32 Idx = Scale; Idx < Count; ++Idx)
 		{
 			const NNECore::Internal::FTensor& CurrTensor = InputTensors[Idx];
-			DmlUtil::FTensorDesc& DmlTensor = DmlInputTensors[Idx];
-
-			DmlTensor = DmlUtil::FTensorDesc{};
+			FTensorDescDml& DmlCurrTensorDesc = DmlInputTensorDescs[Idx];
 
 			if (CurrTensor.GetShape().Rank() == 1)
 			{
-				if (!DmlTensor.InitFromTensor1D(CurrTensor, DmlInputTensors[X].Sizes.Num()))
+				if (!DmlCurrTensorDesc
+						.SetFromTensor1D(CurrTensor, InputTensor.GetShape().Rank())
+						.Validate())
 				{
-					UE_LOG(LogNNE, Warning, TEXT("Failed to initialize tensor(s) for DML inference"));
+					UE_LOG(LogNNE, Error, TEXT("Failed to initialize tensor(s) for DML inference"));
 					return false;
 				}
 			}
 			else
 			{
-				if (!DmlTensor.InitFromTensor(CurrTensor, DmlInputTensors[X].Sizes.Num()))
+				if (!DmlCurrTensorDesc
+						.SetFromTensor(CurrTensor)
+						.Validate())
 				{
-					UE_LOG(LogNNE, Warning, TEXT("Failed to initialize tensor(s) for DML inference"));
+					UE_LOG(LogNNE, Error, TEXT("Failed to initialize tensor(s) for DML inference"));
 					return false;
 				}
 			}
 		}
 		
-		DmlUtil::FTensorDesc	DmlOutputTensor{};
+		FTensorDescDml	DmlOutputTensor;
 
-		if (!DmlOutputTensor.InitFromTensor(OutputTensor, DmlInputTensors[X].Sizes.Num()))
+		if (!DmlOutputTensor
+				.SetFromTensor(OutputTensor)
+				.Validate())
 		{
-			UE_LOG(LogNNE, Warning, TEXT("Failed to initialize tensor(s) for DML inference"));
+			UE_LOG(LogNNE, Error, TEXT("Failed to initialize tensor(s) for DML inference"));
 			return false;
 		}
 
 		DML_BATCH_NORMALIZATION_OPERATOR_DESC	OpDesc{};
 
-		OpDesc.InputTensor = &DmlInputTensors[X].Desc;
-		OpDesc.MeanTensor = &DmlInputTensors[Mean].Desc;
-		OpDesc.VarianceTensor = &DmlInputTensors[Variance].Desc;
-		OpDesc.ScaleTensor = &DmlInputTensors[Scale].Desc;
-		OpDesc.BiasTensor = &DmlInputTensors[Bias].Desc;
-		OpDesc.OutputTensor = &DmlOutputTensor.Desc;
+		OpDesc.InputTensor = DmlInputTensorDescs[X].GetDmlDesc();
+		OpDesc.MeanTensor = DmlInputTensorDescs[Mean].GetDmlDesc();
+		OpDesc.VarianceTensor = DmlInputTensorDescs[Variance].GetDmlDesc();
+		OpDesc.ScaleTensor = DmlInputTensorDescs[Scale].GetDmlDesc();
+		OpDesc.BiasTensor = DmlInputTensorDescs[Bias].GetDmlDesc();
+		OpDesc.OutputTensor = DmlOutputTensor.GetDmlDesc();
 		OpDesc.Spatial = static_cast<BOOL>(1);
 		OpDesc.Epsilon = Epsilon;
 		OpDesc.FusedActivation = nullptr;
diff --git a/Engine/Plugins/Experimental/NNE/Source/NNERuntimeRDG/Private/Dml/NNEDmlOperatorCast.cpp b/Engine/Plugins/Experimental/NNE/Source/NNERuntimeRDG/Private/Dml/NNEDmlOperatorCast.cpp
--- a/Engine/Plugins/Experimental/NNE/Source/NNERuntimeRDG/Private/Dml/NNEDmlOperatorCast.cpp
+++ b/Engine/Plugins/Experimental/NNE/Source/NNERuntimeRDG/Private/Dml/NNEDmlOperatorCast.cpp
@@ -32,7 +32,7 @@ public:
 		
 		if (To != OutputTensor.GetDataType())
 		{
-			UE_LOG(LogNNE, Warning, TEXT("Cast should output a tensor of type %d but was of type %d."), To, OutputTensor.GetDataType());
+			UE_LOG(LogNNE, Error, TEXT("Cast should output a tensor of type %d but was of type %d."), To, OutputTensor.GetDataType());
 			return false;
 		}
 
@@ -41,7 +41,7 @@ public:
 
 		if (InputShape.Num() != OutputShape.Num())
 		{
-			UE_LOG(LogNNE, Warning, TEXT("Cast input and output shapes need to have a same rank"));
+			UE_LOG(LogNNE, Error, TEXT("Cast input and output shapes need to have a same rank"));
 			return false;
 		}
 
@@ -49,31 +49,36 @@ public:
 		{
 			if (InputShape[Idx] != OutputShape[Idx])
 			{
-				UE_LOG(LogNNE, Warning, TEXT("Input shape and output shape need to have a same dimension at dim %d (%d != %d"), Idx, InputShape[Idx], OutputShape[Idx]);
+				UE_LOG(LogNNE, Error, TEXT("Input shape and output shape need to have a same dimension at dim %d (%d != %d"), Idx, InputShape[Idx], OutputShape[Idx]);
 				return false;
 			}
 		}
 
 		// Initialize tensor descriptors
-		DmlUtil::FTensorDesc	DmlInputTensor{};
-		DmlUtil::FTensorDesc	DmlOutputTensor{};
-
-		if (!DmlInputTensor.InitFromTensor(InputTensor, InputTensor.GetShape().Rank()))
+		FTensorDescDml DmlInputTensorDesc;
+		
+		if (!DmlInputTensorDesc
+				.SetFromTensor(InputTensor)
+				.Validate())
 		{
-			UE_LOG(LogNNE, Warning, TEXT("Failed to initialize tensor(s) for DML inference"));
+			UE_LOG(LogNNE, Error, TEXT("Failed to initialize tensor(s) for DML inference"));
 			return false;
 		}
 
-		if (!DmlOutputTensor.InitFromTensor(OutputTensor, InputTensor.GetShape().Rank()))
+		FTensorDescDml DmlOutputTensorDesc;
+
+		if (!DmlOutputTensorDesc
+				.SetFromTensor(OutputTensor)
+				.Validate())
 		{
-			UE_LOG(LogNNE, Warning, TEXT("Failed to initialize tensor(s) for DML inference"));
+			UE_LOG(LogNNE, Error, TEXT("Failed to initialize tensor(s) for DML inference"));
 			return false;
 		}
 
 		DML_CAST_OPERATOR_DESC	OpDesc{};
 
-		OpDesc.InputTensor = &DmlInputTensor.Desc;
-		OpDesc.OutputTensor = &DmlOutputTensor.Desc;
+		OpDesc.InputTensor = DmlInputTensorDesc.GetDmlDesc();
+		OpDesc.OutputTensor = DmlOutputTensorDesc.GetDmlDesc();
 
 		return CreateOperator(Device, DML_OPERATOR_DESC { DML_OPERATOR_CAST, &OpDesc });
 	}
diff --git a/Engine/Plugins/Experimental/NNE/Source/NNERuntimeRDG/Private/Dml/NNEDmlOperatorConcat.cpp b/Engine/Plugins/Experimental/NNE/Source/NNERuntimeRDG/Private/Dml/NNEDmlOperatorConcat.cpp
--- a/Engine/Plugins/Experimental/NNE/Source/NNERuntimeRDG/Private/Dml/NNEDmlOperatorConcat.cpp
+++ b/Engine/Plugins/Experimental/NNE/Source/NNERuntimeRDG/Private/Dml/NNEDmlOperatorConcat.cpp
@@ -34,54 +34,58 @@ public:
 		int32 InputDims = InputTensors[0].GetShape().Rank();
 		auto NormalizeAxis = [InputDims] (int32 Axis) -> uint32
 		{
-			if(Axis < 0)
+			if (Axis < 0)
 			{
 				Axis += InputDims;
 			}
 			return (uint32) Axis;
 		};
+
 		uint32 Axis = NormalizeAxis(Attributes.GetValue<int32>(TEXT("axis")));
 		check(Axis < (uint32) InputDims);
 
-		TArray<DmlUtil::FTensorDesc> InputDescs;
+		TArray<FTensorDescDml> InputDescs;
 		InputDescs.SetNum(InputTensors.Num());
+		
 		TArray<DML_TENSOR_DESC> DmlInputDescs;
 		DmlInputDescs.SetNumUninitialized(InputTensors.Num());
 		
-		for(TConstEnumerateRef<NNECore::Internal::FTensor> It : EnumerateRange(InputTensors))
+		for (TConstEnumerateRef<NNECore::Internal::FTensor> It : EnumerateRange(InputTensors))
 		{
 			const NNECore::Internal::FTensor& Tensor = *It;
 			int Idx = It.GetIndex();
-			if(Algo::Find(Tensor.GetShape().GetData(), 0) == nullptr)
+			
+			if (Algo::Find(Tensor.GetShape().GetData(), 0) == nullptr)
 			{
-				DmlUtil::FTensorDesc& InputTensorDesc = InputDescs[Idx];
-				if (!InputTensorDesc.InitFromTensor(Tensor, Tensor.GetShape().Rank()))
+				if (!InputDescs[Idx]
+						.SetFromTensor(Tensor)
+						.Validate())
 				{
 					UE_LOG(LogNNE, Error, TEXT("Failed to initialize Concat input for DML inference"));
 					return false;
 				}
-				DmlInputDescs[Idx] = InputTensorDesc.Desc;
 			}
 		}
-
+		
 		check(InputDescs.Num() > 0);
 
-		DmlUtil::FTensorDesc OutputTensorDesc;
-		if (!OutputTensorDesc.InitFromTensor(OutputTensors[0], OutputTensors[0].GetShape().Rank()))
+		FTensorDescDml OutputTensorDesc;
+
+		if (!OutputTensorDesc
+				.SetFromTensor(OutputTensors[0])
+				.Validate())
 		{
 			UE_LOG(LogNNE, Error, TEXT("Failed to initialize Concat output for DML inference"));
 			return false;
 		}
 
-
 		DML_JOIN_OPERATOR_DESC DmlJoinOpDesc = {};
 		DmlJoinOpDesc.InputCount = (uint32) DmlInputDescs.Num();
 		DmlJoinOpDesc.InputTensors = DmlInputDescs.GetData();
-		DmlJoinOpDesc.OutputTensor = &OutputTensorDesc.Desc;
+		DmlJoinOpDesc.OutputTensor = OutputTensorDesc.GetDmlDesc();
 		DmlJoinOpDesc.Axis = Axis;
 
 		return CreateOperator(Device, DML_OPERATOR_DESC{ DML_OPERATOR_JOIN, &DmlJoinOpDesc} );
-
 	}
 };
 
diff --git a/Engine/Plugins/Experimental/NNE/Source/NNERuntimeRDG/Private/Dml/NNEDmlOperatorConvolution.cpp b/Engine/Plugins/Experimental/NNE/Source/NNERuntimeRDG/Private/Dml/NNEDmlOperatorConvolution.cpp
--- a/Engine/Plugins/Experimental/NNE/Source/NNERuntimeRDG/Private/Dml/NNEDmlOperatorConvolution.cpp
+++ b/Engine/Plugins/Experimental/NNE/Source/NNERuntimeRDG/Private/Dml/NNEDmlOperatorConvolution.cpp
@@ -18,16 +18,16 @@ class FOperatorDmlConv : public FOperatorDml
 
 	struct FConvArgs
 	{
-		EAutoPad		AutoPad;
-		DmlUtil::FSmallUIntArray StartPadding;
-		DmlUtil::FSmallUIntArray EndPadding;
-		FIntArray		OutPadding;
-		FIntArray		Dilations;
-		FIntArray		Strides;
-		FSmallArray		OutputShape;
-		uint32			NumDimensions;
-		FSmallArray		WindowSize;
-		int32			Group;
+		EAutoPad				AutoPad;
+		Util::FSmallUIntArray	StartPadding;
+		Util::FSmallUIntArray	EndPadding;
+		FIntArray				OutPadding;
+		FIntArray				Dilations;
+		FIntArray				Strides;
+		FSmallArray				OutputShape;
+		uint32					NumDimensions;
+		FSmallArray				WindowSize;
+		int32					Group;
 
 		FConvArgs() = default;
 
@@ -129,7 +129,7 @@ class FOperatorDmlConv : public FOperatorDml
 		//
 		//
 		//
-		DmlUtil::FSmallUIntArray ConvolutionPadding(TConstArrayView<uint32> InputShape)
+		Util::FSmallUIntArray ConvolutionPadding(TConstArrayView<uint32> InputShape)
 		{
 			const uint32 DimOffset = NonspatialDimensionCount;
 
@@ -143,7 +143,7 @@ class FOperatorDmlConv : public FOperatorDml
 			// Deconvolution
 			else
 			{
-				DmlUtil::FSmallUIntArray Padding;
+				Util::FSmallUIntArray Padding;
 				Padding.SetNumUninitialized(NumDimensions);
 				
 				for (uint32 Dim = 0; Dim < NumDimensions; ++Dim)
@@ -230,20 +230,26 @@ public:
 		}
 
 		// Initialize tensor descriptors
-		DmlUtil::FTensorDesc	DmlInputTensor{};
-		DmlUtil::FTensorDesc	DmlFilterTensor{};
-		DmlUtil::FTensorDesc	DmlBiasTensor{};
-		DmlUtil::FTensorDesc	DmlOutputTensor{};
-
-		if (!DmlInputTensor.InitFromTensor(InputTensor, 3))
+		FTensorDescDml	DmlInputTensorDesc;
+		FTensorDescDml	DmlFilterTensorDesc;
+		FTensorDescDml	DmlBiasTensorDesc;
+		FTensorDescDml	DmlOutputTensorDesc;
+
+		if (!DmlInputTensorDesc
+				.SetTensorRank(3, 5)
+				.SetFromTensor(InputTensor)
+				.Validate())
 		{
-			UE_LOG(LogNNE, Warning, TEXT("Failed to initialize tensor(s) for DML inference"));
+			UE_LOG(LogNNE, Error, TEXT("Failed to initialize tensor(s) for DML inference"));
 			return false;
 		}
 
-		if (!DmlFilterTensor.InitFromTensor(FilterTensor, 3))
+		if (!DmlFilterTensorDesc
+				.SetTensorRank(3, 5)
+				.SetFromTensor(InputTensor)
+				.Validate())
 		{
-			UE_LOG(LogNNE, Warning, TEXT("Failed to initialize tensor(s) for DML inference"));
+			UE_LOG(LogNNE, Error, TEXT("Failed to initialize tensor(s) for DML inference"));
 			return false;
 		}
 
@@ -251,43 +257,31 @@ public:
 		{
 			const NNECore::Internal::FTensor& BiasTensor = InputTensors[2];
 
-			FSmallArray Shape;
-
-			if (BiasTensor.GetShape().Rank() < (int32) Args.NumDimensions)
-			{
-				Shape.Add(1);
-				Shape.Add(BiasTensor.GetShape().GetData()[0]);
-				Shape.Add(1);
-
-				for (int32 Dim = 3; Dim < InputTensor.GetShape().Rank(); ++Dim)
-				{
-					Shape.Add(1);
-				}
-			}
-			else
-			{
-				Shape = BiasTensor.GetShape().GetData();
-			}
-
-			if (!DmlBiasTensor.InitFromTensor(BiasTensor, 3, MakeEmptyArrayView<uint32>(), Shape))
+			if (!DmlBiasTensorDesc
+					.SetTensorRank(3, 5)
+					.SetFromTensor1D(BiasTensor, InputTensor.GetShape().Rank())
+					.Validate())
 			{
-				UE_LOG(LogNNE, Warning, TEXT("Failed to initialize tensor(s) for DML inference"));
+				UE_LOG(LogNNE, Error, TEXT("Failed to initialize tensor(s) for DML inference"));
 				return false;
 			}
 		}
 
-		if (!DmlOutputTensor.InitFromTensor(OutputTensor, 3))
+		if (!DmlOutputTensorDesc
+				.SetTensorRank(3, 5)
+				.SetFromTensor(OutputTensor)
+				.Validate())
 		{
-			UE_LOG(LogNNE, Warning, TEXT("Failed to initialize tensor(s) for DML inference"));
+			UE_LOG(LogNNE, Error, TEXT("Failed to initialize tensor(s) for DML inference"));
 			return false;
 		}
 
 		DML_CONVOLUTION_OPERATOR_DESC	DmlConvOpDesc{};
 
-		DmlConvOpDesc.InputTensor = &DmlInputTensor.Desc;
-		DmlConvOpDesc.FilterTensor = &DmlFilterTensor.Desc;
-		DmlConvOpDesc.BiasTensor = InputTensors.Num() > 2 ? &DmlBiasTensor.Desc : nullptr;
-		DmlConvOpDesc.OutputTensor = &DmlOutputTensor.Desc;
+		DmlConvOpDesc.InputTensor = DmlInputTensorDesc.GetDmlDesc();
+		DmlConvOpDesc.FilterTensor = DmlFilterTensorDesc.GetDmlDesc();
+		DmlConvOpDesc.BiasTensor = InputTensors.Num() > 2 ? DmlBiasTensorDesc.GetDmlDesc() : nullptr;
+		DmlConvOpDesc.OutputTensor = DmlOutputTensorDesc.GetDmlDesc();
 		DmlConvOpDesc.Mode = DML_CONVOLUTION_MODE_CROSS_CORRELATION;
 		DmlConvOpDesc.Direction = Direction;
 		DmlConvOpDesc.DimensionCount = Args.NumDimensions;
diff --git a/Engine/Plugins/Experimental/NNE/Source/NNERuntimeRDG/Private/Dml/NNEDmlOperatorCumSum.cpp b/Engine/Plugins/Experimental/NNE/Source/NNERuntimeRDG/Private/Dml/NNEDmlOperatorCumSum.cpp
--- a/Engine/Plugins/Experimental/NNE/Source/NNERuntimeRDG/Private/Dml/NNEDmlOperatorCumSum.cpp
+++ b/Engine/Plugins/Experimental/NNE/Source/NNERuntimeRDG/Private/Dml/NNEDmlOperatorCumSum.cpp
@@ -38,7 +38,7 @@ public:
         case ENNETensorDataType::Int32:
 			{
 				TConstArrayView<int32> TensorContent = InputTensors[1].GetPreparedData<int32>();
-				if(TensorContent.Num() != 1)
+				if (TensorContent.Num() != 1)
 				{
 					UE_LOG(LogNNE, Error, TEXT("axis tensor should be 0-D"));
 					return false;
@@ -46,15 +46,17 @@ public:
 				Axis = (UINT) HandleNegativeAxis(TensorContent[0], InputTensors[0].GetShape().Rank());
 			}
             break;
+
         case ENNETensorDataType::Int64:
 			{
 				TConstArrayView<int64> TensorContent = InputTensors[1].GetPreparedData<int64>();
-				if(TensorContent.Num() != 1)
+				if (TensorContent.Num() != 1)
 				{
 					UE_LOG(LogNNE, Error, TEXT("axis tensor should be 0-D"));
 					return false;
 				}
-				if(DmlUtil::IsOverflowing<int64, int32>(TensorContent[0]))
+				
+				if (Util::IsOverflowing<int64, int32>(TensorContent[0]))
 				{
 					UE_LOG(LogNNE, Error, TEXT("axis number is overflowing"));
 					return false;
@@ -62,23 +64,29 @@ public:
 				Axis = (UINT) HandleNegativeAxis(TensorContent[0], InputTensors[0].GetShape().Rank());
 			}
             break;
+
         default:
             UE_LOG(LogNNE, Error, TEXT("axis tensor has invalid data type"));
 			return false;
         };
 
-		const NNECore::Internal::FTensor& InputTensorDesc = InputTensors[0];
-		const NNECore::Internal::FTensor& OutputTensorDesc = OutputTensors[0];
+		const NNECore::Internal::FTensor& InputTensor = InputTensors[0];
+		const NNECore::Internal::FTensor& OutputTensor = OutputTensors[0];
 
-		DmlUtil::FTensorDesc	DmlInputTensorDesc{};
-		DmlUtil::FTensorDesc	DmlOutputTensorDesc{};
+		FTensorDescDml	DmlInputTensorDesc;
+		FTensorDescDml	DmlOutputTensorDesc;
 
-		if (!InitDmlTensorDesc(DmlInputTensorDesc, InputTensorDesc))
+		if (!DmlInputTensorDesc
+				.SetFromTensor(InputTensor)
+				.Validate())
 		{
 			UE_LOG(LogNNE, Warning, TEXT("Failed to initialize input tensor for DML inference"));
 			return false;
 		}
-		if (!InitDmlTensorDesc(DmlOutputTensorDesc, OutputTensorDesc))
+		
+		if (!DmlOutputTensorDesc
+				.SetFromTensor(OutputTensor)
+				.Validate())
 		{
 			UE_LOG(LogNNE, Warning, TEXT("Failed to initialize output tensor for DML inference"));
 			return false;
@@ -86,8 +94,8 @@ public:
 
 		DML_CUMULATIVE_SUMMATION_OPERATOR_DESC DmlCumSumOpDesc{};
 
-		DmlCumSumOpDesc.InputTensor = &DmlInputTensorDesc.Desc;
-		DmlCumSumOpDesc.OutputTensor = &DmlOutputTensorDesc.Desc;
+		DmlCumSumOpDesc.InputTensor = DmlInputTensorDesc.GetDmlDesc();
+		DmlCumSumOpDesc.OutputTensor = DmlOutputTensorDesc.GetDmlDesc();
 		DmlCumSumOpDesc.Axis = Axis;
 		DmlCumSumOpDesc.AxisDirection = (DML_AXIS_DIRECTION) Attributes.GetValueOrDefault<int32>(TEXT("reverse"), 0);
 		DmlCumSumOpDesc.HasExclusiveSum = (bool) Attributes.GetValueOrDefault<int32>(TEXT("exclusive"), 0);
diff --git a/Engine/Plugins/Experimental/NNE/Source/NNERuntimeRDG/Private/Dml/NNEDmlOperatorDepthToSpace.cpp b/Engine/Plugins/Experimental/NNE/Source/NNERuntimeRDG/Private/Dml/NNEDmlOperatorDepthToSpace.cpp
--- a/Engine/Plugins/Experimental/NNE/Source/NNERuntimeRDG/Private/Dml/NNEDmlOperatorDepthToSpace.cpp
+++ b/Engine/Plugins/Experimental/NNE/Source/NNERuntimeRDG/Private/Dml/NNEDmlOperatorDepthToSpace.cpp
@@ -55,15 +55,22 @@ public:
 		const NNECore::Internal::FTensor& InputTensorDesc = InputTensors[0];
 		const NNECore::Internal::FTensor& OutputTensorDesc = OutputTensors[0];
 
-		DmlUtil::FTensorDesc	DmlInputTensorDesc{};
-		DmlUtil::FTensorDesc	DmlOutputTensorDesc{};
+		FTensorDescDml	DmlInputTensorDesc;
+		FTensorDescDml	DmlOutputTensorDesc;
 
-		if (!DmlInputTensorDesc.InitFromTensor(InputTensorDesc, 4))
+		if (!DmlInputTensorDesc
+				.SetTensorRank(4, 4)
+				.SetFromTensor(InputTensorDesc)
+				.Validate())
 		{
 			UE_LOG(LogNNE, Warning, TEXT("Failed to initialize input tensor for DML inference"));
 			return false;
 		}
-		if (!DmlOutputTensorDesc.InitFromTensor(OutputTensorDesc, 4))
+
+		if (!DmlOutputTensorDesc
+				.SetTensorRank(4, 4)
+				.SetFromTensor(OutputTensorDesc)
+				.Validate())
 		{
 			UE_LOG(LogNNE, Warning, TEXT("Failed to initialize output tensor for DML inference"));
 			return false;
@@ -71,8 +78,8 @@ public:
 
 		DML_DEPTH_TO_SPACE1_OPERATOR_DESC DmlDepthToSpaceOpDesc{};
 
-		DmlDepthToSpaceOpDesc.InputTensor = &DmlInputTensorDesc.Desc;
-		DmlDepthToSpaceOpDesc.OutputTensor = &DmlOutputTensorDesc.Desc;
+		DmlDepthToSpaceOpDesc.InputTensor = DmlInputTensorDesc.GetDmlDesc();
+		DmlDepthToSpaceOpDesc.OutputTensor = DmlOutputTensorDesc.GetDmlDesc();
 		DmlDepthToSpaceOpDesc.BlockSize = BlockSize;
 		DmlDepthToSpaceOpDesc.Order = SpaceOrderFromModeString(Attributes.GetValueOrDefault<FString>(TEXT("mode"), FString(TEXT("DCR"))));
 
diff --git a/Engine/Plugins/Experimental/NNE/Source/NNERuntimeRDG/Private/Dml/NNEDmlOperatorElementwise.cpp b/Engine/Plugins/Experimental/NNE/Source/NNERuntimeRDG/Private/Dml/NNEDmlOperatorElementwise.cpp
--- a/Engine/Plugins/Experimental/NNE/Source/NNERuntimeRDG/Private/Dml/NNEDmlOperatorElementwise.cpp
+++ b/Engine/Plugins/Experimental/NNE/Source/NNERuntimeRDG/Private/Dml/NNEDmlOperatorElementwise.cpp
@@ -39,9 +39,8 @@ public:
 	virtual bool Initialize(IDMLDevice* Device, TArrayView<const NNECore::Internal::FTensor> InputTensors, TArrayView<const NNECore::Internal::FTensor> OutputTensors, const NNECore::FAttributeMap& Attributes) override
 	{
 
-		const NNECore::Internal::FTensor& InputTensorDesc = InputTensors[0];
-		const NNECore::Internal::FTensor& OutputTensorDesc = OutputTensors[0];
-
+		const NNECore::Internal::FTensor& InputTensor = InputTensors[0];
+		const NNECore::Internal::FTensor& OutputTensor = OutputTensors[0];
 
 		if constexpr (std::is_same_v<DmlElementWiseOpDescType, DML_ELEMENT_WISE_CLIP_OPERATOR_DESC>)
 		{
@@ -69,17 +68,19 @@ public:
 		}
 
 		// Initialize tensor descriptor (it's same for both input and output)
-		DmlUtil::FTensorDesc	DmlTensorDesc{};
+		FTensorDescDml	DmlTensorDesc;
 
-		if (!InitDmlTensorDesc(DmlTensorDesc, InputTensorDesc))
+		if (!DmlTensorDesc
+				.SetFromTensor(InputTensor)
+				.Validate())
 		{
-			UE_LOG(LogNNE, Warning, TEXT("Failed to initialize tensor(s) for DML inference"));
+			UE_LOG(LogNNE, Error, TEXT("Failed to initialize tensor(s) for DML inference"));
 			return false;
 		}
 
 		DmlElementWiseOpDescType	DmlElemWiseOpDesc{};
 
-		InitDmlOpDesc(DmlElemWiseOpDesc, DmlTensorDesc);
+		InitDmlOpDesc(DmlElemWiseOpDesc, *DmlTensorDesc.GetDmlDesc());
 
 		DML_OPERATOR_DESC DmlOpDesc{};
 
@@ -92,16 +93,16 @@ public:
 private:
 
 	template<typename OpDesc>
-	void InitDmlOpDesc(OpDesc& Desc, DmlUtil::FTensorDesc& TensorDesc)
+	void InitDmlOpDesc(OpDesc& Desc, const DML_TENSOR_DESC& TensorDesc)
 	{
-		Desc.InputTensor = &TensorDesc.Desc;
-		Desc.OutputTensor = &TensorDesc.Desc;
+		Desc.InputTensor = &TensorDesc;
+		Desc.OutputTensor = &TensorDesc;
 	}
 
-	void InitDmlOpDesc(DML_ELEMENT_WISE_CLIP_OPERATOR_DESC& Desc, DmlUtil::FTensorDesc& TensorDesc)
+	void InitDmlOpDesc(DML_ELEMENT_WISE_CLIP_OPERATOR_DESC& Desc, const DML_TENSOR_DESC& TensorDesc)
 	{
-		Desc.InputTensor = &TensorDesc.Desc;
-		Desc.OutputTensor = &TensorDesc.Desc;
+		Desc.InputTensor = &TensorDesc;
+		Desc.OutputTensor = &TensorDesc;
 		Desc.Min = Min;
 		Desc.Max = Max;
 	}
@@ -137,36 +138,42 @@ public:
 	//
 	virtual bool Initialize(IDMLDevice* Device, TArrayView<const NNECore::Internal::FTensor> InputTensors, TArrayView<const NNECore::Internal::FTensor> OutputTensors, const NNECore::FAttributeMap& Attributes) override
 	{
-		const NNECore::Internal::FTensor& InputATensorDesc = InputTensors[0];
-		const NNECore::Internal::FTensor& InputBTensorDesc = InputTensors[1];
-		const NNECore::Internal::FTensor& OutputTensorDesc = OutputTensors[0];
+		const NNECore::Internal::FTensor& InputATensor = InputTensors[0];
+		const NNECore::Internal::FTensor& InputBTensor = InputTensors[1];
+		const NNECore::Internal::FTensor& OutputTensor = OutputTensors[0];
 
 		// Initialize tensor descriptors
-		DmlUtil::FTensorDesc	DmlInputATensorDesc{};
-		DmlUtil::FTensorDesc	DmlInputBTensorDesc{};
-		DmlUtil::FTensorDesc	DmlOutputTensorDesc{};
+		FTensorDescDml	DmlInputATensorDesc;
+		FTensorDescDml	DmlInputBTensorDesc;
+		FTensorDescDml	DmlOutputTensorDesc;
 
-		if (!InitDmlTensorDesc(DmlInputATensorDesc, InputATensorDesc, OutputTensorDesc))
+		if (!DmlInputATensorDesc
+				.SetFromTensorBroadcast(InputATensor, OutputTensor.GetShape())
+				.Validate())
 		{
-			UE_LOG(LogNNE, Warning, TEXT("Failed to initialize tensor(s) for DML inference"));
+			UE_LOG(LogNNE, Error, TEXT("Failed to initialize tensor(s) for DML inference"));
 			return false;
 		}
 
-		if (!InitDmlTensorDesc(DmlInputBTensorDesc, InputBTensorDesc, OutputTensorDesc))
+		if (!DmlInputBTensorDesc
+				.SetFromTensorBroadcast(InputBTensor, OutputTensor.GetShape())
+				.Validate())
 		{
-			UE_LOG(LogNNE, Warning, TEXT("Failed to initialize tensor(s) for DML inference"));
+			UE_LOG(LogNNE, Error, TEXT("Failed to initialize tensor(s) for DML inference"));
 			return false;
 		}
 
-		if (!InitDmlTensorDesc(DmlOutputTensorDesc, OutputTensorDesc))
+		if (!DmlOutputTensorDesc
+				.SetFromTensor(OutputTensor)
+				.Validate())
 		{
-			UE_LOG(LogNNE, Warning, TEXT("Failed to initialize tensor(s) for DML inference"));
+			UE_LOG(LogNNE, Error, TEXT("Failed to initialize tensor(s) for DML inference"));
 			return false;
 		}
 
 		DmlElementWiseOpDescType	DmlElemWiseOpDesc{};
 
-		InitDmlOpDesc(DmlElemWiseOpDesc, DmlInputATensorDesc, DmlInputBTensorDesc, DmlOutputTensorDesc);
+		InitDmlOpDesc(DmlElemWiseOpDesc, *DmlInputATensorDesc.GetDmlDesc(), *DmlInputBTensorDesc.GetDmlDesc(), *DmlOutputTensorDesc.GetDmlDesc());
 
 		DML_OPERATOR_DESC DmlOpDesc{};
 
@@ -179,18 +186,18 @@ public:
 private:
 
 	template<typename OpDesc>
-	void InitDmlOpDesc(OpDesc& Desc, DmlUtil::FTensorDesc& LHSTensor, DmlUtil::FTensorDesc& RHSTensor, DmlUtil::FTensorDesc& OutputTensor)
+	void InitDmlOpDesc(OpDesc& Desc, const DML_TENSOR_DESC& LHSTensorDesc, const DML_TENSOR_DESC& RHSTensorDesc, const DML_TENSOR_DESC& OutputTensorDesc)
 	{
-		Desc.ATensor = &LHSTensor.Desc;
-		Desc.BTensor = &RHSTensor.Desc;
-		Desc.OutputTensor = &OutputTensor.Desc;
+		Desc.ATensor = &LHSTensorDesc;
+		Desc.BTensor = &RHSTensorDesc;
+		Desc.OutputTensor = &OutputTensorDesc;
 	}
 
-	void InitDmlOpDesc(DML_ELEMENT_WISE_POW_OPERATOR_DESC& Desc, DmlUtil::FTensorDesc& LHSTensor, DmlUtil::FTensorDesc& RHSTensor, DmlUtil::FTensorDesc& OutputTensor)
+	void InitDmlOpDesc(DML_ELEMENT_WISE_POW_OPERATOR_DESC& Desc, const DML_TENSOR_DESC& LHSTensorDesc, const DML_TENSOR_DESC& RHSTensorDesc, const DML_TENSOR_DESC& OutputTensorDesc)
 	{
-		Desc.InputTensor = &LHSTensor.Desc;
-		Desc.ExponentTensor = &RHSTensor.Desc;
-		Desc.OutputTensor = &OutputTensor.Desc;
+		Desc.InputTensor = &LHSTensorDesc;
+		Desc.ExponentTensor = &RHSTensorDesc;
+		Desc.OutputTensor = &OutputTensorDesc;
 	}
 };
 
diff --git a/Engine/Plugins/Experimental/NNE/Source/NNERuntimeRDG/Private/Dml/NNEDmlOperatorFlatten.cpp b/Engine/Plugins/Experimental/NNE/Source/NNERuntimeRDG/Private/Dml/NNEDmlOperatorFlatten.cpp
--- a/Engine/Plugins/Experimental/NNE/Source/NNERuntimeRDG/Private/Dml/NNEDmlOperatorFlatten.cpp
+++ b/Engine/Plugins/Experimental/NNE/Source/NNERuntimeRDG/Private/Dml/NNEDmlOperatorFlatten.cpp
@@ -45,26 +45,27 @@ public:
 			InnerDimSize *= InputShape[Idx];
 		}
 
-		DmlUtil::FSmallUIntArray	OutputShape;
+		Util::FSmallUIntArray	OutputShape;
 
 		OutputShape.Reserve(InputShape.Num());
 		OutputShape.Add(InnerDimSize);
 		OutputShape.Add(InputTensors[0].GetShape().Volume() / InnerDimSize);
 		
-		DmlUtil::FTensorDesc DmlTensorDesc;
+		FTensorDescDml DmlTensorDesc;
 		
-		if (!DmlTensorDesc.InitFromTensor(OutputTensors[0], OutputShape.Num(),
-			/*Broadcast =*/ MakeEmptyArrayView<uint32>(),
-			/*CustomShape =*/ OutputShape))
+		if (!DmlTensorDesc
+				.SetFromTensor(OutputTensors[0])
+				.SetShape(OutputShape)
+				.Validate())
 		{
-			UE_LOG(LogNNE, Warning, TEXT("Failed to initialize Flatten's output tensor for DML inference"));
+			UE_LOG(LogNNE, Error, TEXT("Failed to initialize Flatten's output tensor for DML inference"));
 			return false;
 		}
 
 		DML_ELEMENT_WISE_IDENTITY_OPERATOR_DESC DmlIdentityOpDesc{};
 
-		DmlIdentityOpDesc.InputTensor = &DmlTensorDesc.Desc;
-		DmlIdentityOpDesc.OutputTensor = &DmlTensorDesc.Desc;
+		DmlIdentityOpDesc.InputTensor = DmlTensorDesc.GetDmlDesc();
+		DmlIdentityOpDesc.OutputTensor = DmlTensorDesc.GetDmlDesc();
 
 		return CreateOperator(Device, DML_OPERATOR_DESC{ DML_OPERATOR_ELEMENT_WISE_IDENTITY, &DmlIdentityOpDesc });
 	}
diff --git a/Engine/Plugins/Experimental/NNE/Source/NNERuntimeRDG/Private/Dml/NNEDmlOperatorGRU.cpp b/Engine/Plugins/Experimental/NNE/Source/NNERuntimeRDG/Private/Dml/NNEDmlOperatorGRU.cpp
--- a/Engine/Plugins/Experimental/NNE/Source/NNERuntimeRDG/Private/Dml/NNEDmlOperatorGRU.cpp
+++ b/Engine/Plugins/Experimental/NNE/Source/NNERuntimeRDG/Private/Dml/NNEDmlOperatorGRU.cpp
@@ -131,24 +131,29 @@ public:
 	virtual bool Initialize(IDMLDevice* Device, TArrayView<const NNECore::Internal::FTensor> InputTensors, TArrayView<const NNECore::Internal::FTensor> OutputTensors, const NNECore::FAttributeMap& Attributes) override
 	{
 		
-#define SET_DMLTENSORDESC_FROM_TENSOR_COND(OpDesc, Tensor, Name, Cond)\
-		DmlUtil::FTensorDesc	Dml##Name{};\
+#define SET_DMLTENSORDESC_FROM_TENSOR_COND_DATATYPE(OpDesc, Tensor, Name, DataType, Cond)\
+		FTensorDescDml	Dml##Name##Desc;\
 		if(Cond)\
 		{\
-			if (!Dml##Name.InitFromTensor(Tensor, 4))\
+			if (!Dml##Name##Desc\
+					.SetTensorRank(4, 4)\
+					.SetShape(Tensor.GetShape())\
+					.SetDataType(DataType)\
+					.Validate())\
 			{\
 				UE_LOG(LogNNE, Error, TEXT("Failed to initialize GRU's "#Name" for DML inference"));\
 				return false;\
 			}\
-			##OpDesc.##Name = &Dml##Name.Desc;\
+			##OpDesc.##Name = Dml##Name##Desc.GetDmlDesc();\
 		}\
 		else\
 		{\
 			##OpDesc.##Name = nullptr;\
 		}
 
-#define SET_DMLTENSORDESC_FROM_TENSOR(OpDesc, Tensor, Name)	SET_DMLTENSORDESC_FROM_TENSOR_COND(OpDesc, Tensor, Name, true)
+#define SET_DMLTENSORDESC_FROM_TENSOR_COND(OpDesc, Tensor, Name, Cond)	SET_DMLTENSORDESC_FROM_TENSOR_COND_DATATYPE(OpDesc, Tensor, Name, Tensor.GetDataType(), Cond)
 
+#define SET_DMLTENSORDESC_FROM_TENSOR(OpDesc, Tensor, Name)	SET_DMLTENSORDESC_FROM_TENSOR_COND(OpDesc, Tensor, Name, true)
 
 		check(InputTensors.Num() >= 3 && InputTensors.Num() <= 6);
 		check(OutputTensors.Num() >= 0 && OutputTensors.Num() <= 2);
@@ -224,14 +229,12 @@ public:
 		SET_DMLTENSORDESC_FROM_TENSOR(DmlGRUOpDesc, InputTensors[1], WeightTensor)
 		SET_DMLTENSORDESC_FROM_TENSOR(DmlGRUOpDesc, InputTensors[2], RecurrenceTensor)
 		SET_DMLTENSORDESC_FROM_TENSOR_COND(DmlGRUOpDesc, InputTensors[3], BiasTensor, InputTensors.Num() >= 4)
-		SET_DMLTENSORDESC_FROM_TENSOR_COND(DmlGRUOpDesc, InputTensors[4], SequenceLengthsTensor, InputTensors.Num() >= 5)
 		// Cast SequenceLengthsTensor from int32 to uint32 due to differences in representation between ONNX and DML formats.
-		DmlSequenceLengthsTensor.BuffDesc.DataType = DML_TENSOR_DATA_TYPE::DML_TENSOR_DATA_TYPE_UINT32;
+		SET_DMLTENSORDESC_FROM_TENSOR_COND_DATATYPE(DmlGRUOpDesc, InputTensors[4], SequenceLengthsTensor, ENNETensorDataType::UInt32, InputTensors.Num() >= 5)
 		SET_DMLTENSORDESC_FROM_TENSOR_COND(DmlGRUOpDesc, InputTensors[5], HiddenInitTensor, InputTensors.Num() >= 6)
 		SET_DMLTENSORDESC_FROM_TENSOR_COND(DmlGRUOpDesc, OutputTensors[0], OutputSequenceTensor, OutputTensors.Num() >= 1)
 		SET_DMLTENSORDESC_FROM_TENSOR_COND(DmlGRUOpDesc, OutputTensors[1], OutputSingleTensor, OutputTensors.Num() >= 2)
 
-
 		DmlGRUOpDesc.Direction = DirectionFromString(Attributes.GetValueOrDefault<FString>(TEXT("direction"), TEXT("forward")));
 		
 		const TArray<FString> DefaultActivations = {"Sigmoid", "Tanh"};
@@ -245,14 +248,12 @@ public:
 		check(DmlGRUOpDesc.Direction == DML_RECURRENT_NETWORK_DIRECTION_BIDIRECTIONAL ?
 				(Activations.Num() == 4) : (Activations.Num() == 2));
 
-		
 		if(Attributes.GetAttributeValue(TEXT("clip")) != nullptr)
 		{
 			UE_LOG(LogNNE, Error, TEXT("GRU's clip attribute not supported for DML inference"));
 			return false;
 		}
 		
-		
 		const int HiddenSizeParam = Attributes.GetValueOrDefault<int32>(TEXT("hidden_size"), HiddenSize);
 		check(HiddenSize == HiddenSizeParam);
 
@@ -291,12 +292,8 @@ public:
 		DmlGRUOpDesc.LinearBeforeReset = 
 			Attributes.GetValueOrDefault<int32>(TEXT("linear_before_reset"), 0) == 0 ? false : true;
 
-		DML_OPERATOR_DESC DmlOpDesc{};
-
-		DmlOpDesc.Type = DML_OPERATOR_GRU;
-		DmlOpDesc.Desc = &DmlGRUOpDesc;
+		return CreateOperator(Device, DML_OPERATOR_DESC{ DML_OPERATOR_GRU, &DmlGRUOpDesc });
 
-		return CreateOperator(Device, DmlOpDesc);
 #undef SET_DMLTENSORDESC_FROM_TENSOR
 #undef SET_DMLTENSORDESC_FROM_TENSOR_COND
 	}
diff --git a/Engine/Plugins/Experimental/NNE/Source/NNERuntimeRDG/Private/Dml/NNEDmlOperatorGather.cpp b/Engine/Plugins/Experimental/NNE/Source/NNERuntimeRDG/Private/Dml/NNEDmlOperatorGather.cpp
--- a/Engine/Plugins/Experimental/NNE/Source/NNERuntimeRDG/Private/Dml/NNEDmlOperatorGather.cpp
+++ b/Engine/Plugins/Experimental/NNE/Source/NNERuntimeRDG/Private/Dml/NNEDmlOperatorGather.cpp
@@ -56,7 +56,7 @@ public:
 
 		// Compute output shape
 		const int32 OutputRank = IndicesShape.Rank() + InputShape.Rank() - 1;
-		DmlUtil::FSmallUIntArray OutputShape;
+		Util::FSmallUIntArray OutputShape;
 		int32 DataRankIdx = 0;
 
 		for (; DataRankIdx < Axis; ++DataRankIdx)
@@ -79,33 +79,40 @@ public:
 		}
 
 		// Initialize tensor descriptors
-		DmlUtil::FTensorDesc	DmlInputTensor{};
-		DmlUtil::FTensorDesc	DmlIndicesTensor{};
-		DmlUtil::FTensorDesc	DmlOutputTensor{};
+		FTensorDescDml	DmlInputTensorDesc;
+		FTensorDescDml	DmlIndicesTensorDesc;
+		FTensorDescDml	DmlOutputTensorDesc;
 
-		if (!DmlInputTensor.InitFromTensor(InputTensor, 1))
+		if (!DmlInputTensorDesc
+				.SetFromTensor(InputTensor)
+				.Validate())
 		{
-			UE_LOG(LogNNE, Warning, TEXT("Failed to initialize tensor(s) for DML inference"));
+			UE_LOG(LogNNE, Error, TEXT("Failed to initialize tensor(s) for DML inference"));
 			return false;
 		}
 
-		if (!DmlIndicesTensor.InitFromTensor(IndicesTensor, InputShape.Rank()))
+		if (!DmlIndicesTensorDesc
+				.SetFromTensor(IndicesTensor)
+				.Validate())
 		{
-			UE_LOG(LogNNE, Warning, TEXT("Failed to initialize tensor(s) for DML inference"));
+			UE_LOG(LogNNE, Error, TEXT("Failed to initialize tensor(s) for DML inference"));
 			return false;
 		}
 
-		if (!DmlOutputTensor.InitFromTensor(OutputTensor, InputShape.Rank(), MakeEmptyArrayView<uint32>(), OutputShape))
+		if (!DmlOutputTensorDesc
+				.SetFromTensor(OutputTensor)
+				.SetShape(OutputShape)
+				.Validate())
 		{
-			UE_LOG(LogNNE, Warning, TEXT("Failed to initialize tensor(s) for DML inference"));
+			UE_LOG(LogNNE, Error, TEXT("Failed to initialize tensor(s) for DML inference"));
 			return false;
 		}
 
 		DML_GATHER_OPERATOR_DESC	OpDesc{};
 
-		OpDesc.InputTensor = &DmlInputTensor.Desc;
-		OpDesc.IndicesTensor = &DmlIndicesTensor.Desc;
-		OpDesc.OutputTensor = &DmlOutputTensor.Desc;
+		OpDesc.InputTensor = DmlInputTensorDesc.GetDmlDesc();
+		OpDesc.IndicesTensor = DmlIndicesTensorDesc.GetDmlDesc();
+		OpDesc.OutputTensor = DmlOutputTensorDesc.GetDmlDesc();
 		OpDesc.Axis = Axis;
 		OpDesc.IndexDimensions = IndicesTensor.GetShape().Rank();
 
diff --git a/Engine/Plugins/Experimental/NNE/Source/NNERuntimeRDG/Private/Dml/NNEDmlOperatorGemm.cpp b/Engine/Plugins/Experimental/NNE/Source/NNERuntimeRDG/Private/Dml/NNEDmlOperatorGemm.cpp
--- a/Engine/Plugins/Experimental/NNE/Source/NNERuntimeRDG/Private/Dml/NNEDmlOperatorGemm.cpp
+++ b/Engine/Plugins/Experimental/NNE/Source/NNERuntimeRDG/Private/Dml/NNEDmlOperatorGemm.cpp
@@ -34,64 +34,70 @@ public:
 		TransA = Attributes.GetValueOrDefault(TEXT("transA"), TransA);
 		TransB = Attributes.GetValueOrDefault(TEXT("transB"), TransB);
 		
-		const NNECore::Internal::FTensor& InputATensorDesc = InputTensors[0];
-		const NNECore::Internal::FTensor& InputBTensorDesc = InputTensors[1];
-		const NNECore::Internal::FTensor& OutputTensorDesc = OutputTensors[0];
+		const NNECore::Internal::FTensor& InputATensor = InputTensors[0];
+		const NNECore::Internal::FTensor& InputBTensor = InputTensors[1];
+		const NNECore::Internal::FTensor& OutputTensor = OutputTensors[0];
 
 		// Initialize tensor descriptors
-		DmlUtil::FTensorDesc	DmlInputATensorDesc{};
-		DmlUtil::FTensorDesc	DmlInputBTensorDesc{};
-		DmlUtil::FTensorDesc	DmlInputCTensorDesc{};
-		DmlUtil::FTensorDesc	DmlOutputTensorDesc{};
-
-		if (!InitDmlTensorDesc(DmlInputATensorDesc, InputATensorDesc))
+		FTensorDescDml	DmlInputATensorDesc;
+		FTensorDescDml	DmlInputBTensorDesc;
+		FTensorDescDml	DmlInputCTensorDesc;
+		FTensorDescDml	DmlOutputTensorDesc;
+
+		if (!DmlInputATensorDesc
+				.SetTensorRank(2, 4)
+				.SetFromTensor(InputATensor)
+				.Validate())
 		{
-			UE_LOG(LogNNE, Warning, TEXT("Failed to initialize tensor(s) for DML inference"));
+			UE_LOG(LogNNE, Error, TEXT("Failed to initialize tensor(s) for DML inference"));
 			return false;
 		}
 
-		if (!InitDmlTensorDesc(DmlInputBTensorDesc, InputBTensorDesc))
+		if (!DmlInputBTensorDesc
+				.SetTensorRank(2, 4)
+				.SetFromTensor(InputBTensor)
+				.Validate())
 		{
-			UE_LOG(LogNNE, Warning, TEXT("Failed to initialize tensor(s) for DML inference"));
+			UE_LOG(LogNNE, Error, TEXT("Failed to initialize tensor(s) for DML inference"));
 			return false;
 		}
 
 		if (InputTensors.Num() > 2)
 		{
-			const NNECore::Internal::FTensor& InputCTensorDesc = InputTensors[2];
+			const NNECore::Internal::FTensor& InputCTensor = InputTensors[2];
 
-			if (!InitDmlTensorDesc(DmlInputCTensorDesc, InputCTensorDesc, OutputTensorDesc))
+			if (!DmlInputCTensorDesc
+					.SetTensorRank(2, 4)
+					.SetFromTensorBroadcast(InputCTensor, OutputTensor.GetShape())
+					.Validate())
 			{
-				UE_LOG(LogNNE, Warning, TEXT("Failed to initialize tensor(s) for DML inference"));
+				UE_LOG(LogNNE, Error, TEXT("Failed to initialize tensor(s) for DML inference"));
 				return false;
 			}
 		}
 
-		if (!InitDmlTensorDesc(DmlOutputTensorDesc, OutputTensorDesc))
+		if (!DmlOutputTensorDesc
+				.SetTensorRank(2, 4)
+				.SetFromTensor(OutputTensor)
+				.Validate())
 		{
-			UE_LOG(LogNNE, Warning, TEXT("Failed to initialize tensor(s) for DML inference"));
+			UE_LOG(LogNNE, Error, TEXT("Failed to initialize tensor(s) for DML inference"));
 			return false;
 		}
 
 		DML_GEMM_OPERATOR_DESC	DmlGemmOpDesc{};
 
-		DmlGemmOpDesc.ATensor = &DmlInputATensorDesc.Desc;
-		DmlGemmOpDesc.BTensor = &DmlInputBTensorDesc.Desc;
-		DmlGemmOpDesc.CTensor = InputTensors.Num() > 2 ? &DmlInputCTensorDesc.Desc : nullptr;
-		DmlGemmOpDesc.OutputTensor = &DmlOutputTensorDesc.Desc;
+		DmlGemmOpDesc.ATensor = DmlInputATensorDesc.GetDmlDesc();
+		DmlGemmOpDesc.BTensor = DmlInputBTensorDesc.GetDmlDesc();
+		DmlGemmOpDesc.CTensor = InputTensors.Num() > 2 ? DmlInputCTensorDesc.GetDmlDesc() : nullptr;
+		DmlGemmOpDesc.OutputTensor = DmlOutputTensorDesc.GetDmlDesc();
 		DmlGemmOpDesc.Alpha = Alpha;
 		DmlGemmOpDesc.Beta = Beta;
 		DmlGemmOpDesc.TransA = TransA ? DML_MATRIX_TRANSFORM_TRANSPOSE : DML_MATRIX_TRANSFORM_NONE;
 		DmlGemmOpDesc.TransB = TransB ? DML_MATRIX_TRANSFORM_TRANSPOSE : DML_MATRIX_TRANSFORM_NONE;
 
-		DML_OPERATOR_DESC DmlOpDesc{};
-
-		DmlOpDesc.Type = DML_OPERATOR_GEMM;
-		DmlOpDesc.Desc = &DmlGemmOpDesc;
-
-		return CreateOperator(Device, DmlOpDesc);
+		return CreateOperator(Device, DML_OPERATOR_DESC { DML_OPERATOR_GEMM, &DmlGemmOpDesc });
 	}
-
 };
 
 NNE_DML_REGISTER_OP(Gemm)
diff --git a/Engine/Plugins/Experimental/NNE/Source/NNERuntimeRDG/Private/Dml/NNEDmlOperatorIdentity.cpp b/Engine/Plugins/Experimental/NNE/Source/NNERuntimeRDG/Private/Dml/NNEDmlOperatorIdentity.cpp
--- a/Engine/Plugins/Experimental/NNE/Source/NNERuntimeRDG/Private/Dml/NNEDmlOperatorIdentity.cpp
+++ b/Engine/Plugins/Experimental/NNE/Source/NNERuntimeRDG/Private/Dml/NNEDmlOperatorIdentity.cpp
@@ -26,27 +26,32 @@ public:
 		check(InputTensors.Num() == 1);
 		check(OutputTensors.Num() == 1);
 
-		const NNECore::Internal::FTensor& InputTensorDesc = InputTensors[0];
-		const NNECore::Internal::FTensor& OutputTensorDesc = OutputTensors[0];
+		const NNECore::Internal::FTensor& InputTensor = InputTensors[0];
+		const NNECore::Internal::FTensor& OutputTensor = OutputTensors[0];
 
-		DmlUtil::FTensorDesc	DmlInputTensorDesc{};
-		DmlUtil::FTensorDesc	DmlOutputTensorDesc{};
+		FTensorDescDml	DmlInputTensorDesc;
+		FTensorDescDml	DmlOutputTensorDesc;
 
-		if (!InitDmlTensorDesc(DmlInputTensorDesc, InputTensorDesc))
+		if (!DmlInputTensorDesc
+				.SetFromTensor(InputTensor)
+				.Validate())
 		{
-			UE_LOG(LogNNE, Warning, TEXT("Failed to initialize input tensor for DML inference"));
+			UE_LOG(LogNNE, Error, TEXT("Failed to initialize input tensor for DML inference"));
 			return false;
 		}
-		if (!InitDmlTensorDesc(DmlOutputTensorDesc, OutputTensorDesc))
+		
+		if (!DmlOutputTensorDesc
+				.SetFromTensor(OutputTensor)
+				.Validate())
 		{
-			UE_LOG(LogNNE, Warning, TEXT("Failed to initialize output tensor for DML inference"));
+			UE_LOG(LogNNE, Error, TEXT("Failed to initialize output tensor for DML inference"));
 			return false;
 		}
 
 		DML_ELEMENT_WISE_IDENTITY_OPERATOR_DESC DmlIdentityOpDesc{};
 
-		DmlIdentityOpDesc.InputTensor = &DmlInputTensorDesc.Desc;
-		DmlIdentityOpDesc.OutputTensor = &DmlOutputTensorDesc.Desc;
+		DmlIdentityOpDesc.InputTensor = DmlInputTensorDesc.GetDmlDesc();
+		DmlIdentityOpDesc.OutputTensor = DmlOutputTensorDesc.GetDmlDesc();
 
 		return CreateOperator(Device, DML_OPERATOR_DESC{ DML_OPERATOR_ELEMENT_WISE_IDENTITY, &DmlIdentityOpDesc });
 	}
diff --git a/Engine/Plugins/Experimental/NNE/Source/NNERuntimeRDG/Private/Dml/NNEDmlOperatorInstanceNormalization.cpp b/Engine/Plugins/Experimental/NNE/Source/NNERuntimeRDG/Private/Dml/NNEDmlOperatorInstanceNormalization.cpp
--- a/Engine/Plugins/Experimental/NNE/Source/NNERuntimeRDG/Private/Dml/NNEDmlOperatorInstanceNormalization.cpp
+++ b/Engine/Plugins/Experimental/NNE/Source/NNERuntimeRDG/Private/Dml/NNEDmlOperatorInstanceNormalization.cpp
@@ -38,15 +38,18 @@ public:
 		float	Epsilon = Attributes.GetValueOrDefault<float>(TEXT("epsilon"), DefaultEpsilon);
 
 		// Initialize tensor descriptors
-		DmlUtil::FTensorDesc	DmlInputTensor{};
-		DmlUtil::FTensorDesc	DmlScalingTensor{};
-		DmlUtil::FTensorDesc	DmlBiasTensor{};
-		DmlUtil::FTensorDesc	DmlOutputTensor{};
+		FTensorDescDml	DmlInputTensorDesc;
+		FTensorDescDml	DmlScalingTensorDesc;
+		FTensorDescDml	DmlBiasTensorDesc;
+		FTensorDescDml	DmlOutputTensorDesc;
 			
 		// Make sure that input is 4D
-		if (!DmlInputTensor.InitFromTensor(InputTensor, 4))
+		if (!DmlInputTensorDesc
+				.SetTensorRank(4, 4)
+				.SetFromTensor(InputTensor)
+				.Validate())
 		{
-			UE_LOG(LogNNE, Warning, TEXT("Failed to initialize tensor(s) for DML inference"));
+			UE_LOG(LogNNE, Error, TEXT("Failed to initialize tensor(s) for DML inference"));
 			return false;
 		}
 
@@ -54,9 +57,12 @@ public:
 		{
 			const NNECore::Internal::FTensor& ScaleTensor = InputTensors[1];
 
-			if (!DmlScalingTensor.InitFromTensor1D(ScaleTensor, DmlInputTensor.Sizes.Num()))
+			if (!DmlScalingTensorDesc
+					.SetTensorRank(4, 4)
+					.SetFromTensor(ScaleTensor)
+					.Validate())
 			{
-				UE_LOG(LogNNE, Warning, TEXT("Failed to initialize tensor(s) for DML inference"));
+				UE_LOG(LogNNE, Error, TEXT("Failed to initialize tensor(s) for DML inference"));
 				return false;
 			}
 		}
@@ -65,16 +71,22 @@ public:
 		{
 			const NNECore::Internal::FTensor& BiasTensor = InputTensors[2];
 
-			if (!DmlBiasTensor.InitFromTensor1D(BiasTensor, DmlInputTensor.Sizes.Num()))
+			if (!DmlBiasTensorDesc
+					.SetTensorRank(4, 4)
+					.SetFromTensor(BiasTensor)
+					.Validate())
 			{
-				UE_LOG(LogNNE, Warning, TEXT("Failed to initialize tensor(s) for DML inference"));
+				UE_LOG(LogNNE, Error, TEXT("Failed to initialize tensor(s) for DML inference"));
 				return false;
 			}
 		}
 
-		if (!DmlOutputTensor.InitFromTensor(OutputTensor, 4))
+		if (!DmlOutputTensorDesc
+					.SetTensorRank(4, 4)
+					.SetFromTensor(OutputTensor)
+					.Validate())
 		{
-			UE_LOG(LogNNE, Warning, TEXT("Failed to initialize tensor(s) for DML inference"));
+			UE_LOG(LogNNE, Error, TEXT("Failed to initialize tensor(s) for DML inference"));
 			return false;
 		}
 
@@ -87,7 +99,7 @@ public:
 		// at axis 2 up to the last spatial dimension).
 		const int32 InputRank = DmlInputTensor.Sizes.Num();
 
-		DmlUtil::FSmallUIntArray Axes;
+		Util::FSmallUIntArray Axes;
 		Axes.SetNumUninitialized(NcdhwSpatialDimensionCount);
 
 		for (int32 Dim = 0, Value = InputRank - Axes.Num(); Dim < Axes.Num(); ++Dim, ++Value)
@@ -97,10 +109,10 @@ public:
 
 		DML_MEAN_VARIANCE_NORMALIZATION1_OPERATOR_DESC	OpDesc{};
 
-		OpDesc.InputTensor = &DmlInputTensor.Desc;
-		OpDesc.ScaleTensor = InputTensors.Num() > 1 ? &DmlScalingTensor.Desc : nullptr;
-		OpDesc.BiasTensor = InputTensors.Num() > 2 ? &DmlBiasTensor.Desc : nullptr;
-		OpDesc.OutputTensor = &DmlOutputTensor.Desc;
+		OpDesc.InputTensor = DmlInputTensorDesc.DmlDesc();
+		OpDesc.ScaleTensor = InputTensors.Num() > 1 ? DmlScalingTensorDesc.DmlDesc() : nullptr;
+		OpDesc.BiasTensor = InputTensors.Num() > 2 ? DmlBiasTensorDesc.DmlDesc() : nullptr;
+		OpDesc.OutputTensor = DmlOutputTensorDesc.DmlDesc();
 		OpDesc.AxisCount = Axes.Num();
 		OpDesc.Axes = Axes.GetData();
 		OpDesc.NormalizeVariance = true;
@@ -111,17 +123,17 @@ public:
 
 		DML_MEAN_VARIANCE_NORMALIZATION_OPERATOR_DESC	OpDesc{};
 
-		OpDesc.InputTensor = &DmlInputTensor.Desc;
-		OpDesc.ScaleTensor = InputTensors.Num() > 1 ? &DmlScalingTensor.Desc : nullptr;
-		OpDesc.BiasTensor = InputTensors.Num() > 2 ? &DmlBiasTensor.Desc : nullptr;
-		OpDesc.OutputTensor = &DmlOutputTensor.Desc;
+		OpDesc.InputTensor = DmlInputTensorDesc.GetDmlDesc();
+		OpDesc.ScaleTensor = InputTensors.Num() > 1 ? DmlScalingTensorDesc.GetDmlDesc() : nullptr;
+		OpDesc.BiasTensor = InputTensors.Num() > 2 ? DmlBiasTensorDesc.GetDmlDesc() : nullptr;
+		OpDesc.OutputTensor = DmlOutputTensorDesc.GetDmlDesc();
 		OpDesc.CrossChannel = false;
 		OpDesc.NormalizeVariance = true;
 		OpDesc.Epsilon = Epsilon;
 
 		return CreateOperator(Device, DML_OPERATOR_DESC{ DML_OPERATOR_MEAN_VARIANCE_NORMALIZATION, &OpDesc });
-#endif
 
+#endif
 	}
 };
 
diff --git a/Engine/Plugins/Experimental/NNE/Source/NNERuntimeRDG/Private/Dml/NNEDmlOperatorLRN.cpp b/Engine/Plugins/Experimental/NNE/Source/NNERuntimeRDG/Private/Dml/NNEDmlOperatorLRN.cpp
--- a/Engine/Plugins/Experimental/NNE/Source/NNERuntimeRDG/Private/Dml/NNEDmlOperatorLRN.cpp
+++ b/Engine/Plugins/Experimental/NNE/Source/NNERuntimeRDG/Private/Dml/NNEDmlOperatorLRN.cpp
@@ -42,17 +42,23 @@ public:
 		const float Beta = Attributes.GetValueOrDefault<float>(TEXT("beta"), 0.0f);
 		const float Bias = Attributes.GetValueOrDefault<float>(TEXT("bias"), 0.0f);
 
-		DmlUtil::FTensorDesc	DmlInputTensor{};
+		FTensorDescDml	DmlInputTensorDesc;
 
-		if (!DmlInputTensor.InitFromTensor(InputTensor, InputTensor.GetShape().Rank()))
+		if (!DmlInputTensorDesc
+				.SetTensorRank(4, 4)
+				.SetFromTensor(InputTensor)
+				.Validate())
 		{
 			UE_LOG(LogNNE, Warning, TEXT("Failed to initialize tensor(s) for DML inference"));
 			return false;
 		}
 		
-		DmlUtil::FTensorDesc	DmlOutputTensor{};
+		FTensorDescDml	DmlOutputTensorDesc;
 
-		if (!DmlOutputTensor.InitFromTensor(OutputTensor, DmlInputTensor.Sizes.Num()))
+		if (!DmlOutputTensorDesc
+				.SetTensorRank(4, 4)
+				.SetFromTensor(OutputTensor)
+				.Validate())
 		{
 			UE_LOG(LogNNE, Warning, TEXT("Failed to initialize tensor(s) for DML inference"));
 			return false;
@@ -60,8 +66,8 @@ public:
 
 		DML_LOCAL_RESPONSE_NORMALIZATION_OPERATOR_DESC	OpDesc{};
 
-		OpDesc.InputTensor = &DmlInputTensor.Desc;
-		OpDesc.OutputTensor = &DmlOutputTensor.Desc;
+		OpDesc.InputTensor = DmlInputTensorDesc.GetDmlDesc();
+		OpDesc.OutputTensor = DmlOutputTensorDesc.GetDmlDesc();
 		OpDesc.CrossChannel = true; // ONNX only supports cross-channel
 		OpDesc.LocalSize = Size;
 		OpDesc.Alpha = Alpha;
diff --git a/Engine/Plugins/Experimental/NNE/Source/NNERuntimeRDG/Private/Dml/NNEDmlOperatorLpNormalization.cpp b/Engine/Plugins/Experimental/NNE/Source/NNERuntimeRDG/Private/Dml/NNEDmlOperatorLpNormalization.cpp
--- a/Engine/Plugins/Experimental/NNE/Source/NNERuntimeRDG/Private/Dml/NNEDmlOperatorLpNormalization.cpp
+++ b/Engine/Plugins/Experimental/NNE/Source/NNERuntimeRDG/Private/Dml/NNEDmlOperatorLpNormalization.cpp
@@ -34,7 +34,7 @@ public:
 
 		if (InputTensor.GetShape().Rank() > 8)
 		{
-			UE_LOG(LogNNE, Warning, TEXT("InputTensor rank should be between 1 and 8, got:%d"), InputTensor.GetShape().Rank());
+			UE_LOG(LogNNE, Error, TEXT("InputTensor rank should be between 1 and 8, got:%d"), InputTensor.GetShape().Rank());
 			return false;
 		}
 
@@ -43,29 +43,33 @@ public:
 		check(P >= 1 && P <= 2);
 
 		int32 OnnxAxis = Attributes.GetValueOrDefault<int32>(TEXT("axis"), 0);
-
-		DmlUtil::FTensorDesc	DmlInputTensor{};
-
-		if (!DmlInputTensor.InitFromTensor(InputTensor, InputTensor.GetShape().Rank()))
+		
+		FTensorDescDml	DmlInputTensorDesc;
+		
+		if (!DmlInputTensorDesc
+				.SetFromTensor(InputTensor)
+				.Validate())
 		{
-			UE_LOG(LogNNE, Warning, TEXT("Failed to initialize tensor(s) for DML inference"));
+			UE_LOG(LogNNE, Error, TEXT("Failed to initialize tensor(s) for DML inference"));
 			return false;
 		}
 		
-		DmlUtil::FTensorDesc	DmlOutputTensor{};
-
-		if (!DmlOutputTensor.InitFromTensor(OutputTensor, DmlInputTensor.Sizes.Num()))
+		FTensorDescDml	DmlOutputTensorDesc;
+		
+		if (!DmlOutputTensorDesc
+				.SetFromTensor(OutputTensor)
+				.Validate())
 		{
-			UE_LOG(LogNNE, Warning, TEXT("Failed to initialize tensor(s) for DML inference"));
+			UE_LOG(LogNNE, Error, TEXT("Failed to initialize tensor(s) for DML inference"));
 			return false;
 		}
 
-		uint32 DmlAxis = GetDmlAxis(OnnxAxis, InputTensor.GetShape().Rank(), DmlInputTensor.Sizes.Num());
+		uint32 DmlAxis = GetDmlAxis(OnnxAxis, InputTensor.GetShape().Rank(), DmlInputTensorDesc.GetRank());
 
 		DML_LP_NORMALIZATION_OPERATOR_DESC	OpDesc{};
 
-		OpDesc.InputTensor = &DmlInputTensor.Desc;
-		OpDesc.OutputTensor = &DmlOutputTensor.Desc;
+		OpDesc.InputTensor = DmlInputTensorDesc.GetDmlDesc();
+		OpDesc.OutputTensor = DmlOutputTensorDesc.GetDmlDesc();
 		OpDesc.Axis = DmlAxis;
 		OpDesc.Epsilon = DefaultEpsilon;
 		OpDesc.P = P;
diff --git a/Engine/Plugins/Experimental/NNE/Source/NNERuntimeRDG/Private/Dml/NNEDmlOperatorMatMul.cpp b/Engine/Plugins/Experimental/NNE/Source/NNERuntimeRDG/Private/Dml/NNEDmlOperatorMatMul.cpp
--- a/Engine/Plugins/Experimental/NNE/Source/NNERuntimeRDG/Private/Dml/NNEDmlOperatorMatMul.cpp
+++ b/Engine/Plugins/Experimental/NNE/Source/NNERuntimeRDG/Private/Dml/NNEDmlOperatorMatMul.cpp
@@ -33,32 +33,33 @@ public:
 		const NNECore::Internal::FTensor& OutTensor = OutputTensors[0];
 
 		auto TrimMatrixStackLeadingOnes = [](TConstArrayView<uint32>&& Input) -> TConstArrayView<uint32>
+		{
+			int Start = 0;
+			for (; Start < Input.Num() - 2; ++Start)
 			{
-				int Start = 0;
-				for(; Start < Input.Num() - 2; ++Start)
+				if (Input[Start] != 1)
 				{
-					if(Input[Start] != 1)
-					{
-						break;
-					}
+					break;
 				}
-				return Input.RightChop(Start);
-			};
+			}
+			return Input.RightChop(Start);
+		};
 
-		DmlUtil::FSmallUIntArray ActualShapeA (TrimMatrixStackLeadingOnes(ATensor.GetShape().GetData()));
-		DmlUtil::FSmallUIntArray ActualShapeB (TrimMatrixStackLeadingOnes(BTensor.GetShape().GetData()));
-		DmlUtil::FSmallUIntArray ActualShapeOut (TrimMatrixStackLeadingOnes(OutTensor.GetShape().GetData()));
+		Util::FSmallUIntArray ActualShapeA (TrimMatrixStackLeadingOnes(ATensor.GetShape().GetData()));
+		Util::FSmallUIntArray ActualShapeB (TrimMatrixStackLeadingOnes(BTensor.GetShape().GetData()));
+		Util::FSmallUIntArray ActualShapeOut (TrimMatrixStackLeadingOnes(OutTensor.GetShape().GetData()));
 
 		check(ActualShapeA.Num() > 0);
 		check(ActualShapeB.Num() > 0);
 		check(ActualShapeOut.Num() > 0);
 
-		if(ActualShapeA.Num() == 1)
+		if (ActualShapeA.Num() == 1)
 		{
 			check(ActualShapeOut.Num() == 1);
 			ActualShapeA = {1, ActualShapeA[0]};
 		}
-		if(ActualShapeB.Num() == 1)
+
+		if (ActualShapeB.Num() == 1)
 		{
 			check(ActualShapeOut.Num() == 1);
 			ActualShapeB = {ActualShapeB[0], 1};
@@ -76,37 +77,37 @@ public:
 		bool bBRequiresBroadcasting = false;
 		{
 			int32 MaxRank = FMath::Max(ActualShapeA.Num(), ActualShapeB.Num());
-			for(int Dim = MaxRank - 3; Dim >= 0; Dim--)
+			for (int Dim = MaxRank - 3; Dim >= 0; Dim--)
 			{
 				int AIdx = Dim - MaxRank + ActualShapeA.Num();
 				int BIdx = Dim - MaxRank + ActualShapeB.Num();
 
-				if(bARequiresBroadcasting)
+				if (bARequiresBroadcasting)
 				{
-					if(AIdx >= 0)
+					if (AIdx >= 0)
 					{
 						check(ActualShapeA[AIdx] == 1);
 					}
 				}
 
-				if(bBRequiresBroadcasting)
+				if (bBRequiresBroadcasting)
 				{
-					if(BIdx >= 0)
+					if (BIdx >= 0)
 					{
 						check(ActualShapeB[BIdx] == 1);
 					}
 				}
 
-				if(AIdx >= 0 && BIdx >= 0)
+				if (AIdx >= 0 && BIdx >= 0)
 				{
-					if(ActualShapeA[AIdx] > ActualShapeB[BIdx])
+					if (ActualShapeA[AIdx] > ActualShapeB[BIdx])
 					{
 						check(!bARequiresBroadcasting);
 						check(ActualShapeA[AIdx] % ActualShapeB[BIdx] == 0);
 						bBRequiresBroadcasting = true;
 					}
 
-					if(ActualShapeB[BIdx] > ActualShapeA[AIdx])
+					if (ActualShapeB[BIdx] > ActualShapeA[AIdx])
 					{
 						check(!bBRequiresBroadcasting);
 						check(ActualShapeB[BIdx] % ActualShapeA[AIdx] == 0);
@@ -114,30 +115,27 @@ public:
 					}
 				}
 
-				if(AIdx < 0 && BIdx >= 0 && ActualShapeB[BIdx] != 1)
+				if (AIdx < 0 && BIdx >= 0 && ActualShapeB[BIdx] != 1)
 				{
 					bARequiresBroadcasting = true;
 				}
 
-				if(BIdx < 0 && AIdx >= 0 && ActualShapeA[AIdx] != 1)
+				if (BIdx < 0 && AIdx >= 0 && ActualShapeA[AIdx] != 1)
 				{
 					bBRequiresBroadcasting = true;
 				}
-				
 			}
 		}
 
-		
-		DmlUtil::FSmallUIntArray BroadcastShape(
+		Util::FSmallUIntArray BroadcastShape(
 			bARequiresBroadcasting ? 
 			TConstArrayView<uint32>(ActualShapeB).LeftChop(2) : 
 			TConstArrayView<uint32>(ActualShapeA).LeftChop(2)
 			);
 
-		
 		check(BroadcastShape == TConstArrayView<uint32>(ActualShapeOut).LeftChop(2));
 
-		auto To4DTensor = [](TConstArrayView<uint32> BroadcastShape, TConstArrayView<uint32> MatrixShape) -> DmlUtil::FSmallUIntArray
+		auto To4DTensor = [](TConstArrayView<uint32> BroadcastShape, TConstArrayView<uint32> MatrixShape) -> Util::FSmallUIntArray
 		{
 			check(MatrixShape.Num() == 2);
 
@@ -146,21 +144,21 @@ public:
 			{
 				NumMatrices *= Elem;
 			}
-			return DmlUtil::FSmallUIntArray({1, NumMatrices, MatrixShape[0], MatrixShape[1]});
+			return Util::FSmallUIntArray({1, NumMatrices, MatrixShape[0], MatrixShape[1]});
 		};
 
-		DmlUtil::FSmallUIntArray DmlShapeA = To4DTensor(TConstArrayView<uint32>(ActualShapeA).LeftChop(2), TConstArrayView<uint32>(ActualShapeA).Right(2));
-		DmlUtil::FSmallUIntArray BroadcastShapeA = To4DTensor(BroadcastShape, TConstArrayView<uint32>(ActualShapeA).Right(2));
-		DmlUtil::FSmallUIntArray DmlShapeB = To4DTensor(TConstArrayView<uint32>(ActualShapeB).LeftChop(2), TConstArrayView<uint32>(ActualShapeB).Right(2));
-		DmlUtil::FSmallUIntArray BroadcastShapeB = To4DTensor(BroadcastShape, TConstArrayView<uint32>(ActualShapeB).Right(2));
-		DmlUtil::FSmallUIntArray DmlShapeOut = To4DTensor(BroadcastShape, TConstArrayView<uint32>(ActualShapeOut).Right(2));
+		Util::FSmallUIntArray DmlShapeA = To4DTensor(TConstArrayView<uint32>(ActualShapeA).LeftChop(2), TConstArrayView<uint32>(ActualShapeA).Right(2));
+		Util::FSmallUIntArray BroadcastShapeA = To4DTensor(BroadcastShape, TConstArrayView<uint32>(ActualShapeA).Right(2));
+		Util::FSmallUIntArray DmlShapeB = To4DTensor(TConstArrayView<uint32>(ActualShapeB).LeftChop(2), TConstArrayView<uint32>(ActualShapeB).Right(2));
+		Util::FSmallUIntArray BroadcastShapeB = To4DTensor(BroadcastShape, TConstArrayView<uint32>(ActualShapeB).Right(2));
+		Util::FSmallUIntArray DmlShapeOut = To4DTensor(BroadcastShape, TConstArrayView<uint32>(ActualShapeOut).Right(2));
 
 		check(DmlShapeA[3] == DmlShapeB[2]); // K
 		check(DmlShapeOut[2] == DmlShapeA[2]); // M
 		check(DmlShapeOut[3] == DmlShapeB[3]); // N
 
 		// Another broadcast issue!
-		// The current implementation of DmlUtil::FTensorDesc::InitFromTensor() is not correct with respect to broadcast as it sets
+		// The current implementation of Util::FTensorDesc::InitFromTensor() is not correct with respect to broadcast as it sets
 		// wrong sizes and strides when a dimension to be broadcasted is not of size 1.
 		// Example:
 		// [1, 2, 5, 4] ->(broadcast) [1, 12, 5, 4]
@@ -173,7 +171,7 @@ public:
 		// In our specific case here we can use both the batch and index dimensions to represent non-1 broadcasting, but
 		// we are still limited to cases where the broadcast dimension size is a multiple of the broadcasted one.
 		{
-			auto FixBroadcast = [](DmlUtil::FSmallUIntArray& DmlShape, DmlUtil::FSmallUIntArray& BroadcastArray, uint32 BroadcastFactor)
+			auto FixBroadcast = [](Util::FSmallUIntArray& DmlShape, Util::FSmallUIntArray& BroadcastArray, uint32 BroadcastFactor)
 			{
 				if(DmlShape == BroadcastArray)
 				{
@@ -197,40 +195,45 @@ public:
 		//ATensor.SetShape(DmlShapeA);
 		//BTensor.SetShape(DmlShapeB);
 
-		DmlUtil::FTensorDesc	DmlInputATensorDesc{};
-		DmlUtil::FTensorDesc	DmlInputBTensorDesc{};
-		DmlUtil::FTensorDesc	DmlOutputTensorDesc{};
+		FTensorDescDml	DmlInputATensorDesc;
+		FTensorDescDml	DmlInputBTensorDesc;
+		FTensorDescDml	DmlOutputTensorDesc;
 
-		if (!DmlInputATensorDesc.InitFromTensor(ATensor, DmlShapeA.Num(), 
-			/*Broadcast =*/ BroadcastShapeA, 
-			/*CustomShape =*/ DmlShapeA))
+		if (!DmlInputATensorDesc
+				.SetTensorRank(2, 4)
+				.SetFromTensor(ATensor)
+				.SetShape(DmlShapeA, BroadcastShapeA)
+				.Validate())
 		{
-			UE_LOG(LogNNE, Warning, TEXT("Failed to initialize MatMul's tensor A for DML inference"));
+			UE_LOG(LogNNE, Error, TEXT("Failed to initialize MatMul's tensor A for DML inference"));
 			return false;
 		}
 
-		if (!DmlInputBTensorDesc.InitFromTensor(BTensor, DmlShapeB.Num(), 
-			/*Broadcast =*/ BroadcastShapeB,
-			/*CustomShape =*/ DmlShapeB))
+		if (!DmlInputBTensorDesc
+				.SetTensorRank(2, 4)
+				.SetFromTensor(BTensor)
+				.SetShape(DmlShapeB, BroadcastShapeB)
+				.Validate())
 		{
-			UE_LOG(LogNNE, Warning, TEXT("Failed to initialize MatMul's tensor B for DML inference"));
+			UE_LOG(LogNNE, Error, TEXT("Failed to initialize MatMul's tensor B for DML inference"));
 			return false;
 		}
 
-		if (!DmlOutputTensorDesc.InitFromTensor(OutTensor, DmlShapeOut.Num(), 
-			/*Broadcast =*/ MakeArrayView((uint32*) nullptr, 0), 
-			/*CustomShape =*/ DmlShapeOut))
+		if (!DmlOutputTensorDesc
+				.SetTensorRank(2, 4)
+				.SetShape(DmlShapeOut)
+				.Validate())
 		{
-			UE_LOG(LogNNE, Warning, TEXT("Failed to initialize MatMul's output tensor for DML inference"));
+			UE_LOG(LogNNE, Error, TEXT("Failed to initialize MatMul's output tensor for DML inference"));
 			return false;
 		}
 
 		DML_GEMM_OPERATOR_DESC	DmlGemmOpDesc{};
 
-		DmlGemmOpDesc.ATensor = &DmlInputATensorDesc.Desc;
-		DmlGemmOpDesc.BTensor = &DmlInputBTensorDesc.Desc;
+		DmlGemmOpDesc.ATensor = DmlInputATensorDesc.GetDmlDesc();
+		DmlGemmOpDesc.BTensor = DmlInputBTensorDesc.GetDmlDesc();
 		DmlGemmOpDesc.CTensor = nullptr;
-		DmlGemmOpDesc.OutputTensor = &DmlOutputTensorDesc.Desc;
+		DmlGemmOpDesc.OutputTensor = DmlOutputTensorDesc.GetDmlDesc();
 		DmlGemmOpDesc.Alpha = 1.0f;
 		DmlGemmOpDesc.Beta = 0.0f;
 		DmlGemmOpDesc.TransA = DML_MATRIX_TRANSFORM_NONE;
@@ -239,7 +242,6 @@ public:
 
 		return CreateOperator(Device, DML_OPERATOR_DESC{ DML_OPERATOR_GEMM, &DmlGemmOpDesc} );
 	}
-
 };
 
 NNE_DML_REGISTER_OP(MatMul)
diff --git a/Engine/Plugins/Experimental/NNE/Source/NNERuntimeRDG/Private/Dml/NNEDmlOperatorMaxUnpool.cpp b/Engine/Plugins/Experimental/NNE/Source/NNERuntimeRDG/Private/Dml/NNEDmlOperatorMaxUnpool.cpp
--- a/Engine/Plugins/Experimental/NNE/Source/NNERuntimeRDG/Private/Dml/NNEDmlOperatorMaxUnpool.cpp
+++ b/Engine/Plugins/Experimental/NNE/Source/NNERuntimeRDG/Private/Dml/NNEDmlOperatorMaxUnpool.cpp
@@ -31,38 +31,54 @@ public:
 			ConstantCPUInputs.Add(2);
 		}
 
-		const NNECore::Internal::FTensor& InputTensorDesc = InputTensors[0];
-		const NNECore::Internal::FTensor& IndicesTensorDesc = InputTensors[1];
-		const NNECore::Internal::FTensor& OutputTensorDesc = OutputTensors[0];
-
-		DmlUtil::FTensorDesc	DmlInputTensorDesc;
-		DmlUtil::FTensorDesc	DmlIndicesTensorDesc;
-		DmlUtil::FTensorDesc	DmlOutputTensorDesc;
+		const NNECore::Internal::FTensor& InputTensor = InputTensors[0];
+		const NNECore::Internal::FTensor& IndicesTensor = InputTensors[1];
+		const NNECore::Internal::FTensor& OutputTensor = OutputTensors[0];
+
+		FTensorDescDml	DmlInputTensorDesc;
+		FTensorDescDml	DmlIndicesTensorDesc;
+		FTensorDescDml	DmlOutputTensorDesc;
+
+		if (!DmlInputTensorDesc
+				.SetTensorRank(4, 4)
+				.SetFromTensor(InputTensor)
+				.Validate())
+		{
+			UE_LOG(LogNNE, Error, TEXT("Failed to initialize input tensor for DML inference"));
+			return false;
+		}
 
-		if (!DmlInputTensorDesc.InitFromTensor(InputTensorDesc, 4))
+		// DML required IndicesTensor to be in uint64 format, where ONNX allows int32 and int64
+		if (IndicesTensor.GetDataType() != ENNETensorDataType::Int64)
 		{
-			UE_LOG(LogNNE, Warning, TEXT("Failed to initialize input tensor for DML inference"));
+			UE_LOG(LogNNE, Error, TEXT("DML MaxUnpool requires UInt64, please use Int32 in your ONNX model"));
 			return false;
 		}
-		if (!DmlIndicesTensorDesc.InitFromTensor(IndicesTensorDesc, 4))
+
+		if (!DmlIndicesTensorDesc
+				.SetTensorRank(4, 4)
+				.SetFromTensor(IndicesTensor)
+				.SetDataType(ENNETensorDataType::UInt64)
+				.Validate())
 		{
-			UE_LOG(LogNNE, Warning, TEXT("Failed to initialize indices tensor for DML inference"));
+			UE_LOG(LogNNE, Error, TEXT("Failed to initialize indices tensor for DML inference"));
 			return false;
 		}
-		check(DmlIndicesTensorDesc.BuffDesc.DataType == DML_TENSOR_DATA_TYPE::DML_TENSOR_DATA_TYPE_INT64);
-		// Cast IndicesTensor from int64 to uint64 due to differences in representation between ONNX and DML formats.
-		DmlIndicesTensorDesc.BuffDesc.DataType = DML_TENSOR_DATA_TYPE::DML_TENSOR_DATA_TYPE_UINT64;
-		if (!DmlOutputTensorDesc.InitFromTensor(OutputTensorDesc, 4))
+		
+		if (!DmlOutputTensorDesc
+				.SetTensorRank(4, 4)
+				.SetFromTensor(OutputTensor)
+				.Validate())
 		{
-			UE_LOG(LogNNE, Warning, TEXT("Failed to initialize output tensor for DML inference"));
+			UE_LOG(LogNNE, Error, TEXT("Failed to initialize output tensor for DML inference"));
 			return false;
 		}
 
 		DML_MAX_UNPOOLING_OPERATOR_DESC DmlMaxUnpoolOpDesc{};
 
-		DmlMaxUnpoolOpDesc.InputTensor = &DmlInputTensorDesc.Desc;
-		DmlMaxUnpoolOpDesc.IndicesTensor = &DmlIndicesTensorDesc.Desc;
-		DmlMaxUnpoolOpDesc.OutputTensor = &DmlOutputTensorDesc.Desc;
+		DmlMaxUnpoolOpDesc.InputTensor = DmlInputTensorDesc.GetDmlDesc();
+		DmlMaxUnpoolOpDesc.IndicesTensor = DmlIndicesTensorDesc.GetDmlDesc();
+		DmlMaxUnpoolOpDesc.OutputTensor = DmlOutputTensorDesc.GetDmlDesc();
 
 		return CreateOperator(Device, DML_OPERATOR_DESC{ DML_OPERATOR_MAX_UNPOOLING, &DmlMaxUnpoolOpDesc });
 	}
diff --git a/Engine/Plugins/Experimental/NNE/Source/NNERuntimeRDG/Private/Dml/NNEDmlOperatorMeanVarianceNormalization.cpp b/Engine/Plugins/Experimental/NNE/Source/NNERuntimeRDG/Private/Dml/NNEDmlOperatorMeanVarianceNormalization.cpp
--- a/Engine/Plugins/Experimental/NNE/Source/NNERuntimeRDG/Private/Dml/NNEDmlOperatorMeanVarianceNormalization.cpp
+++ b/Engine/Plugins/Experimental/NNE/Source/NNERuntimeRDG/Private/Dml/NNEDmlOperatorMeanVarianceNormalization.cpp
@@ -65,31 +65,35 @@ public:
 			}
 		}
 
-		DmlUtil::FTensorDesc	DmlInputTensor{};
-
-		if (!DmlInputTensor.InitFromTensor(InputTensor, InputTensor.GetShape().Rank()))
+		FTensorDescDml	DmlInputTensorDesc;
+		
+		if (!DmlInputTensorDesc
+				.SetFromTensor(InputTensor)
+				.Validate())
 		{
-			UE_LOG(LogNNE, Warning, TEXT("Failed to initialize tensor(s) for DML inference"));
+			UE_LOG(LogNNE, Error, TEXT("Failed to initialize tensor(s) for DML inference"));
 			return false;
 		}
 		
-		DmlUtil::FTensorDesc	DmlOutputTensor{};
+		FTensorDescDml	DmlOutputTensorDesc;
 
-		if (!DmlOutputTensor.InitFromTensor(OutputTensor, DmlInputTensor.Sizes.Num()))
+		if (!DmlOutputTensorDesc
+				.SetFromTensor(OutputTensor)
+				.Validate())
 		{
-			UE_LOG(LogNNE, Warning, TEXT("Failed to initialize tensor(s) for DML inference"));
+			UE_LOG(LogNNE, Error, TEXT("Failed to initialize tensor(s) for DML inference"));
 			return false;
 		}
 
-		DmlUtil::FSmallUIntArray DmlAxes;
-		SetDmlAxesFromOnnx(DmlAxes, DmlInputTensor.Sizes.Num(), OnnxAxes);
+		Util::FSmallUIntArray DmlAxes;
+		SetDmlAxesFromOnnx(DmlAxes, DmlInputTensorDesc.GetRank(), OnnxAxes);
 
 		DML_MEAN_VARIANCE_NORMALIZATION1_OPERATOR_DESC	OpDesc{};
 
-		OpDesc.InputTensor = &DmlInputTensor.Desc;
+		OpDesc.InputTensor = DmlInputTensorDesc.GetDmlDesc();
 		OpDesc.ScaleTensor = nullptr;
 		OpDesc.BiasTensor = nullptr;
-		OpDesc.OutputTensor = &DmlOutputTensor.Desc;
+		OpDesc.OutputTensor = DmlOutputTensorDesc.GetDmlDesc();
 		OpDesc.AxisCount = DmlAxes.Num();
 		OpDesc.Axes = DmlAxes.GetData();
 		OpDesc.NormalizeVariance = NormalizeVariance;
diff --git a/Engine/Plugins/Experimental/NNE/Source/NNERuntimeRDG/Private/Dml/NNEDmlOperatorPad.cpp b/Engine/Plugins/Experimental/NNE/Source/NNERuntimeRDG/Private/Dml/NNEDmlOperatorPad.cpp
--- a/Engine/Plugins/Experimental/NNE/Source/NNERuntimeRDG/Private/Dml/NNEDmlOperatorPad.cpp
+++ b/Engine/Plugins/Experimental/NNE/Source/NNERuntimeRDG/Private/Dml/NNEDmlOperatorPad.cpp
@@ -59,7 +59,7 @@ public:
 		{
 			if (!InputTensors[1].HasPreparedData())
 			{
-				UE_LOG(LogNNE, Warning, TEXT("pads is only supported as an attribute or a constant tensor, it is here a variable tensor of name %s."), *InputTensors[1].GetName());
+				UE_LOG(LogNNE, Error, TEXT("pads is only supported as an attribute or a constant tensor, it is here a variable tensor of name %s."), *InputTensors[1].GetName());
 				return false;
 			}
 			Pads.Append(InputTensors[1].GetPreparedData<int64>());
@@ -70,34 +70,36 @@ public:
 			Pads = Attributes.GetValue<TArray<int32>>(TEXT("pads"));
 		}
 
-		
-
 		if (InputTensor.GetShape().Rank() * 2 != Pads.Num())
 		{
-			UE_LOG(LogNNE, Warning, TEXT("pads attribute lenght (%d) should be twice the rank of input X (%d)."), Pads.Num(), InputTensor.GetShape().Rank());
+			UE_LOG(LogNNE, Error, TEXT("pads attribute lenght (%d) should be twice the rank of input X (%d)."), Pads.Num(), InputTensor.GetShape().Rank());
 			return false;
 		}
 		
 		// Initialize tensor descriptors
-		DmlUtil::FTensorDesc	DmlInputTensor{};
-		DmlUtil::FTensorDesc	DmlOutputTensor{};
+		FTensorDescDml	DmlInputTensorDesc;
+		FTensorDescDml	DmlOutputTensorDesc;
 
-		if (!InitDmlTensorDesc(DmlInputTensor, InputTensor))
+		if (!DmlInputTensorDesc
+				.SetFromTensor(InputTensor)
+				.Validate())
 		{
-			UE_LOG(LogNNE, Warning, TEXT("Failed to initialize tensor(s) for DML inference"));
+			UE_LOG(LogNNE, Error, TEXT("Failed to initialize tensor(s) for DML inference"));
 			return false;
 		}
 
-		if (!InitDmlTensorDesc(DmlOutputTensor, OutputTensor))
+		if (!DmlOutputTensorDesc
+				.SetFromTensor(OutputTensor)
+				.Validate())
 		{
-			UE_LOG(LogNNE, Warning, TEXT("Failed to initialize tensor(s) for DML inference"));
+			UE_LOG(LogNNE, Error, TEXT("Failed to initialize tensor(s) for DML inference"));
 			return false;
 		}
 
 		DML_PADDING_OPERATOR_DESC	PadOpDesc{};
 
-		PadOpDesc.InputTensor = &DmlInputTensor.Desc;
-		PadOpDesc.OutputTensor = &DmlOutputTensor.Desc;
+		PadOpDesc.InputTensor = DmlInputTensorDesc.GetDmlDesc();
+		PadOpDesc.OutputTensor = DmlOutputTensorDesc.GetDmlDesc();
 		PadOpDesc.PaddingMode = Mode;
 		PadOpDesc.PaddingValue = Value;
 		PadOpDesc.DimensionCount = Pads.Num() / 2;
diff --git a/Engine/Plugins/Experimental/NNE/Source/NNERuntimeRDG/Private/Dml/NNEDmlOperatorPool.cpp b/Engine/Plugins/Experimental/NNE/Source/NNERuntimeRDG/Private/Dml/NNEDmlOperatorPool.cpp
--- a/Engine/Plugins/Experimental/NNE/Source/NNERuntimeRDG/Private/Dml/NNEDmlOperatorPool.cpp
+++ b/Engine/Plugins/Experimental/NNE/Source/NNERuntimeRDG/Private/Dml/NNEDmlOperatorPool.cpp
@@ -43,12 +43,12 @@ public:
 			return false;
 		}
 
-		DmlUtil::FSmallUIntArray StartPadding, EndPadding, KernelShape, Strides, Dilations;
+		Util::FSmallUIntArray StartPadding, EndPadding, KernelShape, Strides, Dilations;
 
 		Strides.Init(1u, NumSpatialDimensions);
 		if constexpr(!UseGlobalPooling)
 		{
-			if (!DmlUtil::GetArrayAttributeNoOverflow(Attributes.GetAttributeValue(TEXT("strides")), Strides, TConstArrayView<uint32>(Strides)))
+			if (!Util::GetArrayAttributeNoOverflow(Attributes.GetAttributeValue(TEXT("strides")), Strides, TConstArrayView<uint32>(Strides)))
 			{
 				UE_LOG(LogNNE, Error, TEXT("Strides attribute cast led to overflow"));
 				return false;
@@ -59,7 +59,7 @@ public:
 		Dilations.Init(1u, NumSpatialDimensions);
 		if constexpr(!UseGlobalPooling)
 		{
-			if (!DmlUtil::GetArrayAttributeNoOverflow(Attributes.GetAttributeValue(TEXT("dilations")), Dilations, TConstArrayView<uint32>(Dilations)))
+			if (!Util::GetArrayAttributeNoOverflow(Attributes.GetAttributeValue(TEXT("dilations")), Dilations, TConstArrayView<uint32>(Dilations)))
 			{
 				UE_LOG(LogNNE, Error, TEXT("Dilations attribute cast led to overflow"));
 				return false;
@@ -69,11 +69,11 @@ public:
 		
 		if constexpr(UseGlobalPooling)
 		{
-			KernelShape = DmlUtil::FSmallUIntArray{ InputShape.RightChop(NonspatialDimensionCount) };
+			KernelShape = Util::FSmallUIntArray{ InputShape.RightChop(NonspatialDimensionCount) };
 		}
 		else
 		{
-			if (!DmlUtil::GetArrayAttributeNoOverflow(Attributes.GetAttributeValue(TEXT("kernel_shape")), KernelShape))
+			if (!Util::GetArrayAttributeNoOverflow(Attributes.GetAttributeValue(TEXT("kernel_shape")), KernelShape))
 			{
 				UE_LOG(LogNNE, Error, TEXT("kernel_shape attribute cast led to overflow"));
 				return false;
@@ -120,22 +120,29 @@ public:
 			Attributes.GetValueOrDefault<int32>(
 				TEXT("storage_order"),
 				0);
-		if(StorageOrder != 0)
+		if (StorageOrder != 0)
 		{
 			UE_LOG(LogNNE, Error, TEXT("storage_order != 0 is not supported for DML inference"));
 			return false;
 		}
 
+		FTensorDescDml DmlInputTensorDesc;
 
-		DmlUtil::FTensorDesc DmlInputTensorDesc;
-		if (!DmlInputTensorDesc.InitFromTensor(InputTensors[0], 4))
+		if (!DmlInputTensorDesc
+				.SetTensorRank(4, 5)
+				.SetFromTensor(InputTensors[0])
+				.Validate())
 		{
 			UE_LOG(LogNNE, Error, TEXT("Failed to initialize pooling operator's input tensor for DML inference"));
 			return false;
 		}
 
-		DmlUtil::FTensorDesc DmlOutputTensorDesc;
-		if (!DmlOutputTensorDesc.InitFromTensor(OutputTensors[0], 4))
+		FTensorDescDml DmlOutputTensorDesc;
+
+		if (!DmlOutputTensorDesc
+				.SetTensorRank(4, 5)
+				.SetFromTensor(OutputTensors[0])
+				.Validate())
 		{
 			UE_LOG(LogNNE, Error, TEXT("Failed to initialize pooling operator's output tensor for DML inference"));
 			return false;
@@ -143,8 +150,8 @@ public:
 
 		auto FillPoolingDesc = [&](auto& PoolingDesc)
 		{
-			PoolingDesc.InputTensor = &DmlInputTensorDesc.Desc;
-			PoolingDesc.OutputTensor = &DmlOutputTensorDesc.Desc;
+			PoolingDesc.InputTensor = DmlInputTensorDesc.GetDmlDesc();
+			PoolingDesc.OutputTensor = DmlOutputTensorDesc.GetDmlDesc();
 			PoolingDesc.DimensionCount = NumSpatialDimensions;
 			PoolingDesc.WindowSize = KernelShape.GetData();
 			PoolingDesc.Strides = Strides.GetData();
@@ -168,18 +175,24 @@ public:
 
 			if (bHasOutputIndices)
 			{
-				DmlUtil::FTensorDesc DmlIndicesTensorDesc;
-				if (!DmlIndicesTensorDesc.InitFromTensor(OutputTensors[1], 4))
+				FTensorDescDml DmlIndicesTensorDesc;
+				
+				if (!DmlIndicesTensorDesc
+						.SetTensorRank(4, 5)
+						.SetFromTensor(OutputTensors[1])
+						.SetDataType(ENNETensorDataType::UInt64)
+						.Validate())
 				{
 					UE_LOG(LogNNE, Warning, TEXT("Failed to initialize MaxPool's indices output tensor for DML inference"));
 					return false;
 				}
-				DmlIndicesTensorDesc.BuffDesc.DataType = DML_TENSOR_DATA_TYPE::DML_TENSOR_DATA_TYPE_UINT64;
-				DmlMaxPoolOpDesc.OutputIndicesTensor = &DmlIndicesTensorDesc.Desc;
+				
+				DmlMaxPoolOpDesc.OutputIndicesTensor = DmlIndicesTensorDesc.GetDmlDesc();
 			}
 
 			DmlMaxPoolOpDesc.Dilations = Dilations.GetData();
 			FillPoolingDesc(DmlMaxPoolOpDesc);
+
 			return CreateOperator(Device, DML_OPERATOR_DESC{ DML_OPERATOR_MAX_POOLING2, &DmlMaxPoolOpDesc });
 		}
 
diff --git a/Engine/Plugins/Experimental/NNE/Source/NNERuntimeRDG/Private/Dml/NNEDmlOperatorReduce.cpp b/Engine/Plugins/Experimental/NNE/Source/NNERuntimeRDG/Private/Dml/NNEDmlOperatorReduce.cpp
--- a/Engine/Plugins/Experimental/NNE/Source/NNERuntimeRDG/Private/Dml/NNEDmlOperatorReduce.cpp
+++ b/Engine/Plugins/Experimental/NNE/Source/NNERuntimeRDG/Private/Dml/NNEDmlOperatorReduce.cpp
@@ -72,18 +72,20 @@ public:
 		HandleNegativeAxes(Axes, InputTensors[0].GetShape().Rank());
 		HandleEmptyAxes(Axes, InputTensors[0].GetShape().Rank());
 
-		DmlUtil::FTensorDesc DmlInputTensorDesc;
+		FTensorDescDml DmlInputTensorDesc;
 		
-		if (!DmlInputTensorDesc.InitFromTensor(InputTensors[0], InputTensors[0].GetShape().Rank()))
+		if (!DmlInputTensorDesc
+				.SetFromTensor(InputTensors[0])
+				.Validate())
 		{
 			UE_LOG(LogNNE, Warning, TEXT("Failed to initialize Reduce input tensor for DML inference"));
 			return false;
 		}
 
-		DmlUtil::FSmallUIntArray	DmlAxes;
-		DmlUtil::FSmallUIntArray	ReducedDims;
+		Util::FSmallUIntArray	DmlAxes;
+		Util::FSmallUIntArray	ReducedDims;
 
-		ReducedDims.SetNum(DmlInputTensorDesc.Sizes.Num());
+		ReducedDims.SetNum(DmlInputTensorDesc.GetRank());
 
 		for (int32& Dim : Axes)
 		{
@@ -92,7 +94,7 @@ public:
 			DmlAxes.Add(Dim);
 		}
 
-		DmlUtil::FSmallUIntArray	OutputShape;
+		Util::FSmallUIntArray	OutputShape;
 
 		if (!KeepDims)
 		{
@@ -116,11 +118,12 @@ public:
 			OutputShape.Append(OutputTensors[0].GetShape().GetData());
 		}
 
-		DmlUtil::FTensorDesc DmlOutputTensorDesc;
+		FTensorDescDml DmlOutputTensorDesc;
 
-		if (!DmlOutputTensorDesc.InitFromTensor(OutputTensors[0], OutputShape.Num(), 
-			/*Broadcast shape*/ MakeEmptyArrayView<uint32>(),
-			/*Custom shape*/ OutputShape))
+		if (!DmlOutputTensorDesc
+				.SetFromTensor(OutputTensors[0])
+				.SetShape(OutputShape)
+				.Validate())
 		{
 			UE_LOG(LogNNE, Warning, TEXT("Failed to initialize Reduce output tensor for DML inference"));
 			return false;
@@ -130,8 +133,8 @@ public:
         {
             DML_ARGMAX_OPERATOR_DESC OpDesc;
             OpDesc.AxisDirection = (DML_AXIS_DIRECTION) Attributes.GetValueOrDefault<int32>(TEXT("select_last_index"), 0);
-            OpDesc.InputTensor = &DmlInputTensorDesc.Desc;
-			OpDesc.OutputTensor = &DmlOutputTensorDesc.Desc;
+            OpDesc.InputTensor = DmlInputTensorDesc.GetDmlDesc();
+			OpDesc.OutputTensor = DmlOutputTensorDesc.GetDmlDesc();
             OpDesc.Axes = DmlAxes.GetData();
             OpDesc.AxisCount = (uint32) DmlAxes.Num();
 
@@ -141,8 +144,8 @@ public:
         {
             DML_ARGMIN_OPERATOR_DESC OpDesc;
             OpDesc.AxisDirection = (DML_AXIS_DIRECTION) Attributes.GetValueOrDefault<int32>(TEXT("select_last_index"), 0);
-            OpDesc.InputTensor = &DmlInputTensorDesc.Desc;
-			OpDesc.OutputTensor = &DmlOutputTensorDesc.Desc;
+            OpDesc.InputTensor = DmlInputTensorDesc.GetDmlDesc();
+			OpDesc.OutputTensor = DmlOutputTensorDesc.GetDmlDesc();
             OpDesc.Axes = DmlAxes.GetData();
             OpDesc.AxisCount = (uint32) DmlAxes.Num();
 
@@ -152,15 +155,14 @@ public:
 		{
 			DML_REDUCE_OPERATOR_DESC OpDesc{};
 
-			OpDesc.InputTensor = &DmlInputTensorDesc.Desc;
-			OpDesc.OutputTensor = &DmlOutputTensorDesc.Desc;
+			OpDesc.InputTensor = DmlInputTensorDesc.GetDmlDesc();
+			OpDesc.OutputTensor = DmlOutputTensorDesc.GetDmlDesc();
 			OpDesc.Function = ReduceFunc;
 			OpDesc.Axes = DmlAxes.GetData();
 			OpDesc.AxisCount = (uint32) DmlAxes.Num();
 
 			return CreateOperator(Device, DML_OPERATOR_DESC{ DML_OPERATOR_REDUCE, &OpDesc} );
 		}
-		
 	}
 };
 
diff --git a/Engine/Plugins/Experimental/NNE/Source/NNERuntimeRDG/Private/Dml/NNEDmlOperatorResample.cpp b/Engine/Plugins/Experimental/NNE/Source/NNERuntimeRDG/Private/Dml/NNEDmlOperatorResample.cpp
--- a/Engine/Plugins/Experimental/NNE/Source/NNERuntimeRDG/Private/Dml/NNEDmlOperatorResample.cpp
+++ b/Engine/Plugins/Experimental/NNE/Source/NNERuntimeRDG/Private/Dml/NNEDmlOperatorResample.cpp
@@ -29,7 +29,7 @@ void RemoveValuesByIndex(TConstArrayView<uint32> Indices, TArray<TData, TAllocat
 }
 
 // Upsample and Resize operator are implemented as a DML Resample operator
-template < bool IsResize>
+template <bool IsResize>
 class FOperatorDmlResample : public FOperatorDml
 {
 	static DML_INTERPOLATION_MODE ModeFromString(FStringView StringVal)
@@ -89,15 +89,15 @@ public:
 		// Read attributes
 		DML_INTERPOLATION_MODE Mode = ModeFromString(Attributes.GetValue<FString>(TEXT("mode")));
 
-		DmlUtil::FSmallArray<float> InputPixelOffsets, OutputPixelOffsets;
+		Util::FSmallArray<float> InputPixelOffsets, OutputPixelOffsets;
+		
 		InputPixelOffsets.Init(0.5f, InputTensor.GetShape().Rank());
 		OutputPixelOffsets.Init(-0.5f, InputTensor.GetShape().Rank());
 
-		DmlUtil::FSmallArray<float> Scales ( ScaleTensor.GetPreparedData<float>() );
+		Util::FSmallArray<float> Scales ( ScaleTensor.GetPreparedData<float>() );
 
 		if constexpr (IsResize)
 		{
-
 			if (Mode == DML_INTERPOLATION_MODE_NEAREST_NEIGHBOR)
 			{
 				FString NeareastMode = Attributes.GetValueOrDefault<FString>(TEXT("nearest_mode"), TEXT("round_prefer_floor"));
@@ -108,28 +108,28 @@ public:
 			}
 
 			FString CoordinateTransformationMode = Attributes.GetValueOrDefault<FString>(TEXT("coordinate_transformation_mode"), TEXT("half_pixel"));
-			for(int Idx = 0; Idx < InputTensor.GetShape().Rank(); ++Idx)
+			for (int Idx = 0; Idx < InputTensor.GetShape().Rank(); ++Idx)
 			{
 				float LengthResized = (float) OutputTensor.GetShape().GetData()[Idx];
 				float LengthOriginal = (float) InputTensor.GetShape().GetData()[Idx];
 				
-				if(CoordinateTransformationMode == TEXT("align_corners"))
+				if (CoordinateTransformationMode == TEXT("align_corners"))
 				{
 					Scales[Idx] = (LengthResized - 1.0f) / (LengthOriginal - 1.0f);
 					InputPixelOffsets[Idx] = 0.f;
 					OutputPixelOffsets[Idx] = 0.f;
 				}
-				else if(CoordinateTransformationMode == TEXT("asymmetric"))
+				else if (CoordinateTransformationMode == TEXT("asymmetric"))
 				{
 					InputPixelOffsets[Idx] = 0.f;
 					OutputPixelOffsets[Idx] = 0.f;
 				}
-				else if(CoordinateTransformationMode == TEXT("tf_half_pixel_for_nn"))
+				else if (CoordinateTransformationMode == TEXT("tf_half_pixel_for_nn"))
 				{
 					InputPixelOffsets[Idx] = 0.0f;
 					OutputPixelOffsets[Idx] = -0.5f;
 				}
-				else if(CoordinateTransformationMode == TEXT("tf_crop_and_resize"))
+				else if (CoordinateTransformationMode == TEXT("tf_crop_and_resize"))
 				{
 					// NOTE: no tests for this, ORT erroneously puts all 0.0fs in the output tensor in this case.
 					
@@ -149,9 +149,10 @@ public:
 						UE_LOG(LogNNE, Warning, TEXT("roi tensor should contain 2*N entries, where N is rank of X."));
 						return false;
 					}
+					
 					float Start = RoiTensor.GetPreparedData<float>()[Idx];
 					float End = RoiTensor.GetPreparedData<float>()[Idx + InputTensor.GetShape().Rank()];
-					if(LengthResized > 1)
+					if (LengthResized > 1)
 					{
 						Scales[Idx] = (LengthResized - 1.0f) / FMath::Max( (End - Start) * (LengthOriginal - 1.0f), 1.0f );
 						InputPixelOffsets[Idx] = Start * (1.0f - LengthOriginal);
@@ -166,7 +167,7 @@ public:
 				}
 				else
 				{
-					if(CoordinateTransformationMode != TEXT("half_pixel"))
+					if (CoordinateTransformationMode != TEXT("half_pixel"))
 					{
 						UE_LOG(LogNNE, Warning, TEXT("Unsupported coordinate transformation mode:%s, using half_pixel instead"), *CoordinateTransformationMode);
 					}
@@ -176,31 +177,25 @@ public:
 		}
 
 		// Initialize tensor descriptors
-		DmlUtil::FTensorDesc	DmlInputTensor{};
-		DmlUtil::FTensorDesc	DmlOutputTensor{};
+		FTensorDescDml	DmlInputTensorDesc;
+		FTensorDescDml	DmlOutputTensorDesc;
 
-		if (!DmlInputTensor.InitFromTensor(InputTensor, 1))
-		{
-			UE_LOG(LogNNE, Warning, TEXT("Failed to initialize tensor(s) for DML inference"));
-			return false;
-		}
-
-		if (!DmlOutputTensor.InitFromTensor(OutputTensor, 1))
-		{
-			UE_LOG(LogNNE, Warning, TEXT("Failed to initialize tensor(s) for DML inference"));
-			return false;
-		}
+		DmlInputTensorDesc
+			.SetTensorRank(1, 4)
+			.SetFromTensor(InputTensor);
+		
+		DmlOutputTensorDesc
+			.SetTensorRank(1, 4)
+			.SetFromTensor(OutputTensor);
 
 		// Notify DML graph that we have constant CPU tensor inputs
 		for (int i = 1; i < InputTensors.Num(); ++i)
 		{
 			ConstantCPUInputs.Add(i);
 		}
-
-		
 		
 		// Find any useless dimensions of size 1 that occur in both input and output
-		DmlUtil::FSmallUIntArray	SqueezeInds;
+		Util::FSmallUIntArray		SqueezeInds;
 		TConstArrayView<uint32>		InputShape = InputTensor.GetShape().GetData();
 		TConstArrayView<uint32>		OutputShape = OutputTensor.GetShape().GetData();
 
@@ -214,9 +209,9 @@ public:
 
 		if (!SqueezeInds.IsEmpty())
 		{
-			DmlUtil::FSmallUIntArray	SqueezedInputShape(InputShape.GetData(), InputShape.Num());
-			DmlUtil::FSmallUIntArray	SqueezedOutputShape(OutputShape.GetData(), OutputShape.Num());
-			DmlUtil::FSmallArray<float> ScaleValues(Scales.GetData(), Scales.Num());
+			Util::FSmallUIntArray		SqueezedInputShape(InputShape.GetData(), InputShape.Num());
+			Util::FSmallUIntArray		SqueezedOutputShape(OutputShape.GetData(), OutputShape.Num());
+			Util::FSmallArray<float>	ScaleValues(Scales.GetData(), Scales.Num());
 
 			RemoveValuesByIndex(SqueezeInds, SqueezedInputShape, true);
 			RemoveValuesByIndex(SqueezeInds, SqueezedOutputShape, true);
@@ -224,15 +219,27 @@ public:
 			RemoveValuesByIndex(SqueezeInds, InputPixelOffsets, true);
 			RemoveValuesByIndex(SqueezeInds, OutputPixelOffsets, true);
 
-			DmlInputTensor.UpdateShapeAndStrides(SqueezedInputShape);
-			DmlOutputTensor.UpdateShapeAndStrides(SqueezedOutputShape);
+			DmlInputTensorDesc.SetShape(SqueezedInputShape);
+			DmlOutputTensorDesc.SetShape(SqueezedOutputShape);
 			Scales = ScaleValues;
 		}
 
+		if (!DmlInputTensorDesc.Validate())
+		{
+			UE_LOG(LogNNE, Error, TEXT("Failed to initialize tensor(s) for DML inference"));
+			return false;
+		}
+
+		if (!DmlOutputTensorDesc.Validate())
+		{
+			UE_LOG(LogNNE, Error, TEXT("Failed to initialize tensor(s) for DML inference"));
+			return false;
+		}
+
 		DML_RESAMPLE1_OPERATOR_DESC	OpDesc{};
 
-		OpDesc.InputTensor = &DmlInputTensor.Desc;
-		OpDesc.OutputTensor = &DmlOutputTensor.Desc;
+		OpDesc.InputTensor = DmlInputTensorDesc.GetDmlDesc();
+		OpDesc.OutputTensor = DmlOutputTensorDesc.GetDmlDesc();
 		OpDesc.InterpolationMode = Mode;
 		OpDesc.DimensionCount = Scales.Num();
 		OpDesc.Scales = Scales.GetData();
diff --git a/Engine/Plugins/Experimental/NNE/Source/NNERuntimeRDG/Private/Dml/NNEDmlOperatorReshape.cpp b/Engine/Plugins/Experimental/NNE/Source/NNERuntimeRDG/Private/Dml/NNEDmlOperatorReshape.cpp
--- a/Engine/Plugins/Experimental/NNE/Source/NNERuntimeRDG/Private/Dml/NNEDmlOperatorReshape.cpp
+++ b/Engine/Plugins/Experimental/NNE/Source/NNERuntimeRDG/Private/Dml/NNEDmlOperatorReshape.cpp
@@ -40,7 +40,7 @@ public:
         bool bAllowZero = (bool)
 			( Attributes.GetValueOrDefault<int32>(TEXT("allowzero"), 0) );
 
-        DmlUtil::FSmallUIntArray ReshapedShape;
+        Util::FSmallUIntArray ReshapedShape;
 
         switch(InputTensors[1].GetDataType())
         {
@@ -69,17 +69,22 @@ public:
         
         check(ReshapedShape == OutputTensors[0].GetShape().GetData());
 
-        DmlUtil::FTensorDesc DmlInputTensorDesc;
-        if (!DmlInputTensorDesc.InitFromTensor(InputTensors[0], ReshapedShape.Num(), 
-			/*Broadcast =*/ MakeArrayView((uint32*) nullptr, 0), 
-			/*CustomShape =*/ ReshapedShape))
+        FTensorDescDml DmlInputTensorDesc;
+
+        if (!DmlInputTensorDesc
+				.SetFromTensor(InputTensors[0])
+				.SetShape(ReshapedShape)
+				.Validate())
 		{
 			UE_LOG(LogNNE, Warning, TEXT("Failed to initialize Reshape's input tensor for DML inference"));
 			return false;
 		}
 
-        DmlUtil::FTensorDesc DmlOutputTensorDesc;
-        if (!DmlOutputTensorDesc.InitFromTensor(OutputTensors[0], OutputTensors[0].GetShape().Rank()))
+        FTensorDescDml DmlOutputTensorDesc;
+
+        if (!DmlOutputTensorDesc
+				.SetFromTensor(OutputTensors[0])
+				.Validate())
 		{
 			UE_LOG(LogNNE, Warning, TEXT("Failed to initialize Reshape's output tensor for DML inference"));
 			return false;
@@ -87,11 +92,10 @@ public:
         
 		DML_ELEMENT_WISE_IDENTITY_OPERATOR_DESC DmlIdentityOpDesc{};
 
-        DmlIdentityOpDesc.InputTensor = &DmlInputTensorDesc.Desc;
-        DmlIdentityOpDesc.OutputTensor = &DmlOutputTensorDesc.Desc;
+        DmlIdentityOpDesc.InputTensor = DmlInputTensorDesc.GetDmlDesc();
+        DmlIdentityOpDesc.OutputTensor = DmlOutputTensorDesc.GetDmlDesc();
 
 		return CreateOperator(Device, DML_OPERATOR_DESC{ DML_OPERATOR_ELEMENT_WISE_IDENTITY, &DmlIdentityOpDesc} );
-
 	}
 };
 
diff --git a/Engine/Plugins/Experimental/NNE/Source/NNERuntimeRDG/Private/Dml/NNEDmlOperatorSlice.cpp b/Engine/Plugins/Experimental/NNE/Source/NNERuntimeRDG/Private/Dml/NNEDmlOperatorSlice.cpp
--- a/Engine/Plugins/Experimental/NNE/Source/NNERuntimeRDG/Private/Dml/NNEDmlOperatorSlice.cpp
+++ b/Engine/Plugins/Experimental/NNE/Source/NNERuntimeRDG/Private/Dml/NNEDmlOperatorSlice.cpp
@@ -20,18 +20,18 @@ class FOperatorDmlSlice : public FOperatorDml
 	template<typename DataType>
 	void ComputeOffsetsSizesStrides(
 		TArrayView<const NNECore::Internal::FTensor> InputTensors,
-		DmlUtil::FSmallUIntArray& OutOffsets, 
-		DmlUtil::FSmallUIntArray& OutSizes, 
-		DmlUtil::FSmallIntArray& OutStrides)
+		Util::FSmallUIntArray& OutOffsets, 
+		Util::FSmallUIntArray& OutSizes, 
+		Util::FSmallIntArray& OutStrides)
 	{
 		TConstArrayView<DataType> Starts =  InputTensors[1].GetPreparedData<DataType>();
 		TConstArrayView<DataType> Ends = InputTensors[2].GetPreparedData<DataType>();
 		check(Starts.Num() == Ends.Num());
 
-		DmlUtil::FSmallArray<DataType> Axes;
-		DmlUtil::FSmallArray<DataType> Steps;
+		Util::FSmallArray<DataType> Axes;
+		Util::FSmallArray<DataType> Steps;
 
-		if(InputTensors.Num() >= 4)
+		if (InputTensors.Num() >= 4)
 		{
 			auto NormalizeAxes = [NumDims = (DataType) InputTensors[0].GetShape().Rank()] (DataType& Axis)
 			{
@@ -50,7 +50,7 @@ class FOperatorDmlSlice : public FOperatorDml
 			std::iota(Axes.begin(), Axes.end(), 0);
 		}
 
-		if(InputTensors.Num() >= 5)
+		if (InputTensors.Num() >= 5)
 		{
 			Steps = InputTensors[4].GetPreparedData<DataType>();
 			check(Steps.Num() == Axes.Num());
@@ -60,7 +60,7 @@ class FOperatorDmlSlice : public FOperatorDml
 			Steps.Init((DataType)1, Axes.Num());
 		}
 
-		DmlUtil::FSmallUIntArray OutputShape;
+		Util::FSmallUIntArray OutputShape;
 		Algo::Copy(InputTensors[0].GetShape().GetData(), OutputShape);
 		OutSizes = OutputShape;
 		OutOffsets.SetNumZeroed(OutputShape.Num());
@@ -79,11 +79,11 @@ class FOperatorDmlSlice : public FOperatorDml
 			check(Stride != 0);
 
 			uint32 Dim = InputTensors[0].GetShape().GetData()[DimIndex];
-			if(Start < 0 && Start > TNumericLimits<DataType>::Min())
+			if (Start < 0 && Start > TNumericLimits<DataType>::Min())
 			{
 				Start += (DataType) Dim;
 			}
-			if(End < 0 && Start > TNumericLimits<DataType>::Min())
+			if (End < 0 && Start > TNumericLimits<DataType>::Min())
 			{
 				End += (DataType) Dim;
 			}
@@ -128,7 +128,7 @@ public:
 		for(int Idx = 1; Idx < InputTensors.Num(); Idx++)
 		{
 			check(InputTensors[Idx].GetShape().Rank() == 1);
-			if(Idx == 1)
+			if (Idx == 1)
 			{
 				InputIndexDataType = InputTensors[Idx].GetDataType();
 				check( InputIndexDataType == ENNETensorDataType::Int32 
@@ -143,24 +143,29 @@ public:
 		}
 
 		// Initialize Input tensor desc
-        DmlUtil::FTensorDesc DmlInputTensorDesc;
-        if (!DmlInputTensorDesc.InitFromTensor(InputTensors[0], InputTensors[0].GetShape().Rank()))
+        FTensorDescDml DmlInputTensorDesc;
+        
+		if (!DmlInputTensorDesc
+				.SetFromTensor(InputTensors[0])
+				.Validate())
         {
             UE_LOG(LogNNE, Error, TEXT("Failed to initialize Slice input for DML inference"));
             return false;
         }
 
 		// Initialize Output tensor desc
-        DmlUtil::FTensorDesc DmlOutputTensorDesc;
-        if (!DmlOutputTensorDesc.InitFromTensor(OutputTensors[0], OutputTensors[0].GetShape().Rank()))
+        FTensorDescDml DmlOutputTensorDesc;
+        
+		if (!DmlOutputTensorDesc
+				.SetFromTensor(OutputTensors[0])
+				.Validate())
         {
             UE_LOG(LogNNE, Error, TEXT("Failed to initialize Slice Output for DML inference"));
             return false;
         }
 
-
-		DmlUtil::FSmallUIntArray Offsets, Sizes;
-		DmlUtil::FSmallIntArray Strides;
+		Util::FSmallUIntArray Offsets, Sizes;
+		Util::FSmallIntArray Strides;
 
 		switch(InputIndexDataType)
 		{
@@ -172,11 +177,10 @@ public:
 			break;
 		}
 
-
 		DML_SLICE1_OPERATOR_DESC DmlSliceOpDesc{};
 		
-		DmlSliceOpDesc.InputTensor = &DmlInputTensorDesc.Desc;
-        DmlSliceOpDesc.OutputTensor = &DmlOutputTensorDesc.Desc;
+		DmlSliceOpDesc.InputTensor = DmlInputTensorDesc.GetDmlDesc();
+        DmlSliceOpDesc.OutputTensor = DmlOutputTensorDesc.GetDmlDesc();
         DmlSliceOpDesc.DimensionCount = (uint32) Offsets.Num();
         DmlSliceOpDesc.InputWindowOffsets = Offsets.GetData();
         DmlSliceOpDesc.InputWindowSizes = Sizes.GetData();
diff --git a/Engine/Plugins/Experimental/NNE/Source/NNERuntimeRDG/Private/Dml/NNEDmlOperatorSpaceToDepth.cpp b/Engine/Plugins/Experimental/NNE/Source/NNERuntimeRDG/Private/Dml/NNEDmlOperatorSpaceToDepth.cpp
--- a/Engine/Plugins/Experimental/NNE/Source/NNERuntimeRDG/Private/Dml/NNEDmlOperatorSpaceToDepth.cpp
+++ b/Engine/Plugins/Experimental/NNE/Source/NNERuntimeRDG/Private/Dml/NNEDmlOperatorSpaceToDepth.cpp
@@ -41,18 +41,23 @@ public:
 		}
 
 
-		const NNECore::Internal::FTensor& InputTensorDesc = InputTensors[0];
-		const NNECore::Internal::FTensor& OutputTensorDesc = OutputTensors[0];
+		const NNECore::Internal::FTensor& InputTensor = InputTensors[0];
+		const NNECore::Internal::FTensor& OutputTensor = OutputTensors[0];
 
-		DmlUtil::FTensorDesc	DmlInputTensorDesc{};
-		DmlUtil::FTensorDesc	DmlOutputTensorDesc{};
+		FTensorDescDml	DmlInputTensorDesc;
+		FTensorDescDml	DmlOutputTensorDesc;
 
-		if (!DmlInputTensorDesc.InitFromTensor(InputTensorDesc, 4))
+		if (!DmlInputTensorDesc
+				.SetFromTensor(InputTensor)
+				.Validate())
 		{
 			UE_LOG(LogNNE, Warning, TEXT("Failed to initialize input tensor for DML inference"));
 			return false;
 		}
-		if (!DmlOutputTensorDesc.InitFromTensor(OutputTensorDesc, 4))
+
+		if (!DmlOutputTensorDesc
+				.SetFromTensor(OutputTensor)
+				.Validate())
 		{
 			UE_LOG(LogNNE, Warning, TEXT("Failed to initialize output tensor for DML inference"));
 			return false;
@@ -60,8 +65,8 @@ public:
 
 		DML_SPACE_TO_DEPTH1_OPERATOR_DESC DmlSpaceToDepthOpDesc{};
 
-		DmlSpaceToDepthOpDesc.InputTensor = &DmlInputTensorDesc.Desc;
-		DmlSpaceToDepthOpDesc.OutputTensor = &DmlOutputTensorDesc.Desc;
+		DmlSpaceToDepthOpDesc.InputTensor = DmlInputTensorDesc.GetDmlDesc();
+		DmlSpaceToDepthOpDesc.OutputTensor = DmlOutputTensorDesc.GetDmlDesc();
 		DmlSpaceToDepthOpDesc.BlockSize = BlockSize;
 		DmlSpaceToDepthOpDesc.Order = DML_DEPTH_SPACE_ORDER_DEPTH_COLUMN_ROW;
 
diff --git a/Engine/Plugins/Experimental/NNE/Source/NNERuntimeRDG/Private/Dml/NNEDmlOperatorSplit.cpp b/Engine/Plugins/Experimental/NNE/Source/NNERuntimeRDG/Private/Dml/NNEDmlOperatorSplit.cpp
--- a/Engine/Plugins/Experimental/NNE/Source/NNERuntimeRDG/Private/Dml/NNEDmlOperatorSplit.cpp
+++ b/Engine/Plugins/Experimental/NNE/Source/NNERuntimeRDG/Private/Dml/NNEDmlOperatorSplit.cpp
@@ -35,11 +35,12 @@ public:
 		for (int Idx = 0; Idx < OutputTensors.Num(); ++Idx)
 		{
 			//check(OutputTensors[Idx].GetShape().Rank() == InputTensor.GetShape().Rank());
-			if(OutputTensors[Idx].GetShape().Rank() != InputTensor.GetShape().Rank())
+			if (OutputTensors[Idx].GetShape().Rank() != InputTensor.GetShape().Rank())
 			{
 				UE_LOG(LogNNE, Error, TEXT("Rank of output tensor and input tensor should be the same"));
 				return false;
 			}
+
 			for (int Dim = 0; Dim < InputTensor.GetShape().Rank(); ++Dim)
 			{
 				if (Dim == Axis)
@@ -49,7 +50,7 @@ public:
 				else
 				{
 					//check(OutputTensors[Idx].GetShape().GetData()[Dim] == InputTensor.GetShape().GetData()[Dim]);
-					if(OutputTensors[Idx].GetShape().GetData()[Dim] != InputTensor.GetShape().GetData()[Dim])
+					if (OutputTensors[Idx].GetShape().GetData()[Dim] != InputTensor.GetShape().GetData()[Dim])
 					{
 						UE_LOG(LogNNE, Error, TEXT("%s"), *FString::Printf(TEXT("Output tensor %d 's dimension %d should match input tensor's"), Idx, Dim));
 						return false;
@@ -59,38 +60,44 @@ public:
 		}
 
 		//check(SplitSize == InputTensor.GetShape().GetData()[Axis]);
-		if(SplitSize != InputTensor.GetShape().GetData()[Axis])
+		if (SplitSize != InputTensor.GetShape().GetData()[Axis])
 		{
 			UE_LOG(LogNNE, Error, TEXT("Input tensor's axis dimension size must be equal to the sum of output tensors' axis dimensions sizes"));
 			return false;
 		}
 
-
-		DmlUtil::FTensorDesc InputTensorDesc;
-		if (!InputTensorDesc.InitFromTensor(InputTensor, InputTensor.GetShape().Rank()))
+		FTensorDescDml InputTensorDesc;
+		
+		if (!InputTensorDesc
+				.SetFromTensor(InputTensor)
+				.Validate())
 		{
 			UE_LOG(LogNNE, Error, TEXT("Failed to initialize pooling operator's input tensor for DML inference"));
 			return false;
 		}
 
-		DmlUtil::FSmallArray<DmlUtil::FTensorDesc> OutputTensorDescs;
-		DmlUtil::FSmallArray<DML_TENSOR_DESC> DmlOutputTensorDescs;
+		Util::FSmallArray<FTensorDescDml> OutputTensorDescs;
+		Util::FSmallArray<DML_TENSOR_DESC> DmlOutputTensorDescs;
+		
 		OutputTensorDescs.SetNum(OutputTensors.Num());
 		DmlOutputTensorDescs.SetNumUninitialized(OutputTensors.Num());
 
 		for(int Idx = 0; Idx < OutputTensors.Num(); ++Idx)
 		{
-			if (!OutputTensorDescs[Idx].InitFromTensor(OutputTensors[Idx], OutputTensors[Idx].GetShape().Rank()))
+			if (!OutputTensorDescs[Idx]
+					.SetFromTensor(OutputTensors[Idx])
+					.Validate())
 			{
 				UE_LOG(LogNNE, Error, TEXT("Failed to initialize pooling operator's output tensor for DML inference"));
 				return false;
 			}
-			DmlOutputTensorDescs[Idx] = OutputTensorDescs[Idx].Desc;
+
+			DmlOutputTensorDescs[Idx] = *OutputTensorDescs[Idx].GetDmlDesc();
 		}
 		
 		DML_SPLIT_OPERATOR_DESC	SplitOpDesc{};
 
-		SplitOpDesc.InputTensor = &InputTensorDesc.Desc;
+		SplitOpDesc.InputTensor = InputTensorDesc.GetDmlDesc();
 		SplitOpDesc.OutputCount = DmlOutputTensorDescs.Num();
 		SplitOpDesc.OutputTensors = DmlOutputTensorDescs.GetData();
 		SplitOpDesc.Axis = (UINT) Axis;
diff --git a/Engine/Plugins/Experimental/NNE/Source/NNERuntimeRDG/Private/Dml/NNEDmlOperatorSqueeze.cpp b/Engine/Plugins/Experimental/NNE/Source/NNERuntimeRDG/Private/Dml/NNEDmlOperatorSqueeze.cpp
--- a/Engine/Plugins/Experimental/NNE/Source/NNERuntimeRDG/Private/Dml/NNEDmlOperatorSqueeze.cpp
+++ b/Engine/Plugins/Experimental/NNE/Source/NNERuntimeRDG/Private/Dml/NNEDmlOperatorSqueeze.cpp
@@ -48,7 +48,7 @@ public:
 			}
 		}
 		
-		DmlUtil::FSmallUIntArray	OutputShape;
+		Util::FSmallUIntArray	OutputShape;
 
 		OutputShape.Reserve(InputShape.Num());
 		OutputShape.Append(InputShape);
@@ -66,11 +66,12 @@ public:
 			OutputShape.RemoveAt(Axe);
 		}
 
-		DmlUtil::FTensorDesc DmlTensorDesc;
+		FTensorDescDml DmlTensorDesc;
 		
-		if (!DmlTensorDesc.InitFromTensor(OutputTensors[0], OutputShape.Num(),
-			/*Broadcast =*/ MakeEmptyArrayView<uint32>(),
-			/*CustomShape =*/ OutputShape))
+		if (!DmlTensorDesc
+				.SetFromTensor(OutputTensors[0])
+				.SetShape(OutputShape)
+				.Validate())
 		{
 			UE_LOG(LogNNE, Warning, TEXT("Failed to initialize Unsqueeze's output tensor for DML inference"));
 			return false;
@@ -78,8 +79,8 @@ public:
 
 		DML_ELEMENT_WISE_IDENTITY_OPERATOR_DESC DmlIdentityOpDesc{};
 
-		DmlIdentityOpDesc.InputTensor = &DmlTensorDesc.Desc;
-		DmlIdentityOpDesc.OutputTensor = &DmlTensorDesc.Desc;
+		DmlIdentityOpDesc.InputTensor = DmlTensorDesc.GetDmlDesc();
+		DmlIdentityOpDesc.OutputTensor = DmlTensorDesc.GetDmlDesc();
 
 		return CreateOperator(Device, DML_OPERATOR_DESC{ DML_OPERATOR_ELEMENT_WISE_IDENTITY, &DmlIdentityOpDesc });
 	}
diff --git a/Engine/Plugins/Experimental/NNE/Source/NNERuntimeRDG/Private/Dml/NNEDmlOperatorTranspose.cpp b/Engine/Plugins/Experimental/NNE/Source/NNERuntimeRDG/Private/Dml/NNEDmlOperatorTranspose.cpp
--- a/Engine/Plugins/Experimental/NNE/Source/NNERuntimeRDG/Private/Dml/NNEDmlOperatorTranspose.cpp
+++ b/Engine/Plugins/Experimental/NNE/Source/NNERuntimeRDG/Private/Dml/NNEDmlOperatorTranspose.cpp
@@ -28,64 +28,73 @@ public:
 	//
 	virtual bool Initialize(IDMLDevice* Device, TArrayView<const NNECore::Internal::FTensor> InputTensors, TArrayView<const NNECore::Internal::FTensor> OutputTensors, const NNECore::FAttributeMap& Attributes) override
 	{
-        check(InputTensors.Num() == 1);
-        check(OutputTensors.Num() == 1);
-
-        int32 NumDims = InputTensors[0].GetShape().Rank();
-        check(NumDims > 0);
-        check(NumDims == OutputTensors[0].GetShape().Rank());
-
-        // Default permutation is reverse
-        DmlUtil::FSmallIntArray ReversePerm;
-        for(int Idx = NumDims-1; Idx >= 0; Idx--)
-        {
-            ReversePerm.Add(Idx);
-        }
-        TArray<int32> Perm = 
+		check(InputTensors.Num() == 1);
+		check(OutputTensors.Num() == 1);
+
+		int32 NumDims = InputTensors[0].GetShape().Rank();
+		check(NumDims > 0);
+		check(NumDims == OutputTensors[0].GetShape().Rank());
+
+		// Default permutation is reverse
+		Util::FSmallIntArray ReversePerm;
+
+		for(int Idx = NumDims-1; Idx >= 0; Idx--)
+		{
+			ReversePerm.Add(Idx);
+		}
+
+		TArray<int32> Perm = 
 			Attributes.GetValueOrDefault<TArray<int32>>(TEXT("perm"), (TArray<int32>) ReversePerm);
 
-        check(Perm.Num() == NumDims);
-
-        // Initialize Input tensor desc
-        DmlUtil::FTensorDesc DmlInputTensorDesc;
-        if (!DmlInputTensorDesc.InitFromTensor(InputTensors[0], NumDims))
-        {
-            UE_LOG(LogNNE, Error, TEXT("Failed to initialize Transpose input for DML inference"));
-            return false;
-        }
-        
-        DmlInputTensorDesc.SetStridesFromTensor(InputTensors[0]);
-
-        // Initialize Output tensor desc
-        DmlUtil::FTensorDesc DmlOutputTensorDesc;
-        if (!DmlOutputTensorDesc.InitFromTensor(OutputTensors[0], NumDims))
-        {
-            UE_LOG(LogNNE, Error, TEXT("Failed to initialize Transpose output for DML inference"));
-            return false;
-        }
-
-        // Apply permutations to both sizes and strides
-        auto PermuteFunc = [&, Perm] (TConstArrayView<uint32> InputView) -> DmlUtil::FSmallUIntArray
-            {
-                DmlUtil::FSmallUIntArray Permuted;
-                for(int32 PermVal : Perm)
-                {
-                    Permuted.Add(InputView[PermVal]);
-                }
-                return Permuted;
-            };
-        
-        DmlInputTensorDesc.UpdateShapeAndStrides(PermuteFunc(DmlInputTensorDesc.Sizes), PermuteFunc(DmlInputTensorDesc.Strides));
-        
-        check(DmlOutputTensorDesc.Sizes == DmlInputTensorDesc.Sizes);        
+		check(Perm.Num() == NumDims);
+
+		// Initialize Input tensor desc
+		// Apply permutations to both sizes and strides
+		auto PermuteFunc = [&, Perm] (TConstArrayView<uint32> InputView) -> Util::FSmallUIntArray
+		{
+			Util::FSmallUIntArray Permuted;
+			
+			for (int32 PermVal : Perm)
+			{
+				Permuted.Add(InputView[PermVal]);
+			}
+			
+			return Permuted;
+		};
+
+		FTensorDescDml DmlInputTensorDesc;
+
+		DmlInputTensorDesc.SetFromTensor(InputTensors[0]);
+		DmlInputTensorDesc.SetShape(PermuteFunc(DmlInputTensorDesc.GetSizes()));
+
+		DmlInputTensorDesc.SetStridesFromShape(InputTensors[0].GetShape());
+		DmlInputTensorDesc.SetStrides(PermuteFunc(DmlInputTensorDesc.GetStrides()));
+
+		if (DmlInputTensorDesc.Validate())
+		{
+			UE_LOG(LogNNE, Error, TEXT("Failed to initialize Transpose input for DML inference"));
+			return false;
+		}
+
+		// Initialize Output tensor desc
+		FTensorDescDml DmlOutputTensorDesc;
+
+		if (!DmlOutputTensorDesc
+				.SetFromTensor(OutputTensors[0])
+				.Validate())
+		{
+			UE_LOG(LogNNE, Error, TEXT("Failed to initialize Transpose output for DML inference"));
+			return false;
+		}
+
+		check(Util::IsSameShape(DmlOutputTensorDesc, DmlInputTensorDesc));
 
 		DML_ELEMENT_WISE_IDENTITY_OPERATOR_DESC DmlIdentityOpDesc{};
 
-        DmlIdentityOpDesc.InputTensor = &DmlInputTensorDesc.Desc;
-        DmlIdentityOpDesc.OutputTensor = &DmlOutputTensorDesc.Desc;
+		DmlIdentityOpDesc.InputTensor = DmlInputTensorDesc.GetDmlDesc();
+		DmlIdentityOpDesc.OutputTensor = DmlOutputTensorDesc.GetDmlDesc();
 
 		return CreateOperator(Device, DML_OPERATOR_DESC{ DML_OPERATOR_ELEMENT_WISE_IDENTITY, &DmlIdentityOpDesc} );
-
 	}
 };
 
diff --git a/Engine/Plugins/Experimental/NNE/Source/NNERuntimeRDG/Private/Dml/NNEDmlOperatorUnsqueeze.cpp b/Engine/Plugins/Experimental/NNE/Source/NNERuntimeRDG/Private/Dml/NNEDmlOperatorUnsqueeze.cpp
--- a/Engine/Plugins/Experimental/NNE/Source/NNERuntimeRDG/Private/Dml/NNEDmlOperatorUnsqueeze.cpp
+++ b/Engine/Plugins/Experimental/NNE/Source/NNERuntimeRDG/Private/Dml/NNEDmlOperatorUnsqueeze.cpp
@@ -40,7 +40,7 @@ public:
 		HandleNegativeAxes(Axes, OutShapeRank);
 		Axes.Sort();
 
-		DmlUtil::FSmallUIntArray	OutputShape;
+		Util::FSmallUIntArray	OutputShape;
 
 		OutputShape.Reserve(OutShapeRank);
 		OutputShape.Append(InputShape);
@@ -56,11 +56,12 @@ public:
 			OutputShape.Insert(1, Axes[Idx]);
 		}
 
-		DmlUtil::FTensorDesc DmlTensorDesc;
+		FTensorDescDml DmlTensorDesc;
 		
-		if (!DmlTensorDesc.InitFromTensor(OutputTensors[0], OutputShape.Num(),
-			/*Broadcast =*/ MakeEmptyArrayView<uint32>(),
-			/*CustomShape =*/ OutputShape))
+		if (!DmlTensorDesc
+				.SetFromTensor(OutputTensors[0])
+				.SetShape(OutputShape)
+				.Validate())
 		{
 			UE_LOG(LogNNE, Warning, TEXT("Failed to initialize Unsqueeze's output tensor for DML inference"));
 			return false;
@@ -68,8 +69,8 @@ public:
 
 		DML_ELEMENT_WISE_IDENTITY_OPERATOR_DESC DmlIdentityOpDesc{};
 
-		DmlIdentityOpDesc.InputTensor = &DmlTensorDesc.Desc;
-		DmlIdentityOpDesc.OutputTensor = &DmlTensorDesc.Desc;
+		DmlIdentityOpDesc.InputTensor = DmlTensorDesc.GetDmlDesc();
+		DmlIdentityOpDesc.OutputTensor = DmlTensorDesc.GetDmlDesc();
 
 		return CreateOperator(Device, DML_OPERATOR_DESC{ DML_OPERATOR_ELEMENT_WISE_IDENTITY, &DmlIdentityOpDesc });
 	}
diff --git a/Engine/Plugins/Experimental/NNE/Source/NNERuntimeRDG/Private/Dml/NNEDmlOperatorUtils.h b/Engine/Plugins/Experimental/NNE/Source/NNERuntimeRDG/Private/Dml/NNEDmlOperatorUtils.h
--- a/Engine/Plugins/Experimental/NNE/Source/NNERuntimeRDG/Private/Dml/NNEDmlOperatorUtils.h
+++ b/Engine/Plugins/Experimental/NNE/Source/NNERuntimeRDG/Private/Dml/NNEDmlOperatorUtils.h
@@ -39,7 +39,7 @@ static EAutoPad AutoPadFromString(FStringView StringVal)
     }
 }
 
-static DmlUtil::FSmallUIntArray KernelPadding(
+static Util::FSmallUIntArray KernelPadding(
     TConstArrayView<uint32> InputShape, TConstArrayView<uint32> WindowSize, 
     TConstArrayView<uint32> Dilations, TConstArrayView<uint32> Strides
     )
@@ -47,7 +47,7 @@ static DmlUtil::FSmallUIntArray KernelPadding(
     const uint32 NumSpatialDimensions = InputShape.Num() - NonspatialDimensionCount;
     check(NumSpatialDimensions >= 1);
 
-    DmlUtil::FSmallUIntArray Padding;
+    Util::FSmallUIntArray Padding;
     Padding.SetNumUninitialized(NumSpatialDimensions);
 
     for (uint32 Dim = 0; Dim < NumSpatialDimensions; ++Dim)
@@ -66,8 +66,8 @@ static DmlUtil::FSmallUIntArray KernelPadding(
 static void ComputeStartEndPaddings(
     TConstArrayView<uint32> InputShape,
     const NNECore::FAttributeMap& Attributes, 
-    DmlUtil::FSmallUIntArray &OutStartPadding, 
-    DmlUtil::FSmallUIntArray &OutEndPadding,
+    Util::FSmallUIntArray &OutStartPadding, 
+	Util::FSmallUIntArray &OutEndPadding,
     TConstArrayView<uint32> Padding)
 {
     const uint32 NumSpatialDimensions = InputShape.Num() - NonspatialDimensionCount;
@@ -77,10 +77,10 @@ static void ComputeStartEndPaddings(
 
     if (AutoPad == EAutoPad::NOTSET)
     {
-        DmlUtil::FSmallUIntArray Pads;
+        Util::FSmallUIntArray Pads;
 
         Pads.Init(0u, 2 * NumSpatialDimensions);
-        check(DmlUtil::GetArrayAttributeNoOverflow(Attributes.GetAttributeValue(TEXT("pads")), Pads, TConstArrayView<uint32>(Pads)))
+        check(Util::GetArrayAttributeNoOverflow(Attributes.GetAttributeValue(TEXT("pads")), Pads, TConstArrayView<uint32>(Pads)))
 
 
         for (uint32 Dim = 0; Dim < NumSpatialDimensions; ++Dim)
@@ -143,7 +143,7 @@ inline int32 GetDmlAxis(int32 OnnxAxis, int32 OnnxDim, int32 DmlDim)
 	return DmlAxis;
 }
 
-inline void SetDmlAxesFromOnnx(DmlUtil::FSmallUIntArray& DmlAxes, int32 Rank, TConstArrayView<int32> OnnxAxes)
+inline void SetDmlAxesFromOnnx(Util::FSmallUIntArray& DmlAxes, int32 Rank, TConstArrayView<int32> OnnxAxes)
 {
 	DmlAxes.Reset();
 	DmlAxes.Reserve(Rank);
