Add RDG DirectML Celu and Shrink ops

#rnx

diff --git a/Engine/Plugins/Experimental/NNE/Source/NNEQA/Resources/NNEQATestDesc.json b/Engine/Plugins/Experimental/NNE/Source/NNEQA/Resources/NNEQATestDesc.json
--- a/Engine/Plugins/Experimental/NNE/Source/NNEQA/Resources/NNEQATestDesc.json
+++ b/Engine/Plugins/Experimental/NNE/Source/NNEQA/Resources/NNEQATestDesc.json
@@ -273,6 +273,10 @@
 				{	"target": "Atanh" },
 				{	"target": "BitShift", "skip": true, "comment": "TODO need attributes" },
 				{	"target": "Ceil" },
+				{	"target": "Celu", 
+					"runtimes": [{ "name": "NNERuntimeRDGHlsl",   "skip": true }],
+					"additional_datasets": [ "Elementwise.Elu.AttributeTest" ]
+				},
 				{	"target": "Clip", "additional_datasets": [ "Elementwise.Clip.AttributeTest" ] },
 				{	"target": "Cos" },
 				{	"target": "Cosh" },
@@ -314,6 +318,10 @@
 						{ "name": "NNERuntimeRDGDml",    "absolute_tolerance":  1e-4},
 						{ "name": "NNERuntimeRDGHlsl",   "absolute_tolerance":  1e-4}
 					]},
+				{	"target": "Shrink", 
+					"runtimes": [{ "name": "NNERuntimeRDGHlsl",   "skip": true }],
+					"additional_datasets": [ "Elementwise.Shrink.AttributeTest" ]
+				},
 				{	"target": "Sigmoid"},
 				{	"target": "Sign" },
 				{	"target": "Sin", "runtimes": [
@@ -734,6 +742,14 @@
 				}
 			]
 		},
+		{
+			"name": "Elementwise.Shrink.AttributeTest",
+			"datasets": [
+				{
+					"inputs": [ {"shape": [ 4 ] } ]
+				}
+			]
+		},
 		{
 			"name": "Elementwise.Softmax.AttributeTest",
 			"datasets": [
@@ -1511,6 +1527,16 @@
 				{ "attributes": [ ] }
 			]
 		},
+		{
+			"name": "Elementwise.Shrink.AttributeTest",
+			"attribute_maps": [
+				{ "attributes": [ { "name": "bias", "value": 0.5 } ] },
+				{ "attributes": [ { "name": "bias", "value": 1.0 } ] },
+				{ "attributes": [ { "name": "bias", "value": 0.5 }, { "name": "lambd", "value": 0.1 } ] },
+				{ "attributes": [ { "name": "bias", "value": 1.0 }, { "name": "lambd", "value": 0.2 } ] },
+				{ "attributes": [ ] }
+			]
+		},
 		{
 			"name": "Elementwise.Softmax.AttributeTest",
 			"attribute_maps": [
diff --git a/Engine/Plugins/Experimental/NNE/Source/NNERuntimeRDG/Private/Dml/NNEDmlOperatorActivations.cpp b/Engine/Plugins/Experimental/NNE/Source/NNERuntimeRDG/Private/Dml/NNEDmlOperatorActivations.cpp
--- a/Engine/Plugins/Experimental/NNE/Source/NNERuntimeRDG/Private/Dml/NNEDmlOperatorActivations.cpp
+++ b/Engine/Plugins/Experimental/NNE/Source/NNERuntimeRDG/Private/Dml/NNEDmlOperatorActivations.cpp
@@ -31,6 +31,8 @@ private:
 	float Alpha;
 	float Beta;
 	float Gamma;
+	float Bias;
+	float Treshold;
 	int Axis;
 
 public:
@@ -54,6 +56,11 @@ public:
 		Beta = Attributes.GetValueOrDefault(TEXT("beta"), Beta);
 		Gamma = Attributes.GetValueOrDefault(TEXT("gamma"), Gamma);
 
+		if constexpr (DmlActivationOpType == DML_OPERATOR_ACTIVATION_SHRINK)
+		{
+			Bias = Attributes.GetValueOrDefault(TEXT("bias"), 0.0f);
+			Treshold = Attributes.GetValueOrDefault(TEXT("lambd"), 0.5f);
+		}
 
 		// Initialize tensor descriptor (it's same for both input and output)
 		FTensorDescDml	DmlTensorDesc;
@@ -139,6 +146,21 @@ private:
 		Desc.OutputTensor = &TensorDesc;
 		Desc.Alpha = Alpha;
 	}
+
+	void InitDmlOpDesc(DML_ACTIVATION_CELU_OPERATOR_DESC& Desc, const DML_TENSOR_DESC& TensorDesc)
+	{
+		Desc.InputTensor = &TensorDesc;
+		Desc.OutputTensor = &TensorDesc;
+		Desc.Alpha = Alpha;
+	}
+
+	void InitDmlOpDesc(DML_ACTIVATION_SHRINK_OPERATOR_DESC& Desc, const DML_TENSOR_DESC& TensorDesc)
+	{
+		Desc.InputTensor = &TensorDesc;
+		Desc.OutputTensor = &TensorDesc;
+		Desc.Bias = Bias;
+		Desc.Threshold = Treshold;
+	}
 };
 
 /**
@@ -263,12 +285,14 @@ REGISTER_OP_ACTIVATION_UNARY(OpName, DmlOpName)
 // 	REGISTER_OP_ACTIVATION_UNARY_PARAMS(	OpName, 			DmlOpName, 				InAlpha, 							InBeta, 	InGamma )
 
 	REGISTER_OP_ACTIVATION_UNARY(			Dropout,			IDENTITY )
-	REGISTER_OP_ACTIVATION_UNARY_PARAMS(	Elu, 				ELU, 					1.0f, 								0.0f, 		1.05070102214813232421875f )
+	REGISTER_OP_ACTIVATION_UNARY_PARAMS(	Celu, 				CELU, 					1.0f, 								0.0f, 		0.0f )
+	REGISTER_OP_ACTIVATION_UNARY_PARAMS(	Elu, 				ELU, 					1.0f, 								0.0f, 		0.0f )
 	REGISTER_OP_ACTIVATION_UNARY_PARAMS(	HardSigmoid, 		HARD_SIGMOID, 			0.2f, 								0.5f, 		0.0f )
 	REGISTER_OP_ACTIVATION_UNARY_PARAMS(	LeakyRelu, 			LEAKY_RELU, 			0.01f, 								0.0f, 		0.0f )
 	REGISTER_OP_ACTIVATION_UNARY(			LogSoftmax, 		LOG_SOFTMAX1 )
 	REGISTER_OP_ACTIVATION_UNARY(			Relu, 				RELU )
 	REGISTER_OP_ACTIVATION_UNARY_PARAMS(	Selu, 				SCALED_ELU, 			1.67326319217681884765625f, 		0.0f, 		1.05070102214813232421875f )
+	REGISTER_OP_ACTIVATION_UNARY(			Shrink, 			SHRINK	)
 	REGISTER_OP_ACTIVATION_UNARY(			Sigmoid, 			SIGMOID	)
 	REGISTER_OP_ACTIVATION_UNARY(			Softmax, 			SOFTMAX1 )
 	REGISTER_OP_ACTIVATION_UNARY(			Softplus, 			SOFTPLUS )
