Fixed mismatched copy constructor/assignment operator for UE::NNECore::FSymbolicTensorShape and UE::NNECore::FTensorShape.

diff --git a/Engine/Plugins/Experimental/NNE/Source/NNECore/Public/NNECoreTypes.h b/Engine/Plugins/Experimental/NNE/Source/NNECore/Public/NNECoreTypes.h
--- a/Engine/Plugins/Experimental/NNE/Source/NNECore/Public/NNECoreTypes.h
+++ b/Engine/Plugins/Experimental/NNE/Source/NNECore/Public/NNECoreTypes.h
@@ -52,6 +52,13 @@ namespace UE::NNECore
 		TArray<int32, TInlineAllocator<MaxRank>> Data;
 
 	public:
+		/**
+		 * Construct this symbolic tensor shape from another one.
+		 * 
+		 * @param OtherShape the other symbolic tensor shape to copy from.
+		 */
+		FSymbolicTensorShape(const FSymbolicTensorShape& OtherShape) = default;
+
 		/**
 		 * Construct a symbolic tensor shape with the given dimensions.
 		 *
@@ -122,6 +129,13 @@ namespace UE::NNECore
 		TArray<uint32, TInlineAllocator<MaxRank>> Data;
 
 	public:
+
+		/**
+		 * Construct this tensor shape from another one.
+		 * @param OtherShape the other tensor shape to copy from.
+		 */
+		FTensorShape(const FTensorShape& OtherShape) = default;
+
 		/**
 		 * Construct a tensor shape with the given dimensions.
 		 *
