Fixed missing default constructors in UE::NNECore::FSymbolicTensorShape and UE::NNECore::FTensorShape.

diff --git a/Engine/Plugins/Experimental/NNE/Source/NNECore/Public/NNECoreTypes.h b/Engine/Plugins/Experimental/NNE/Source/NNECore/Public/NNECoreTypes.h
--- a/Engine/Plugins/Experimental/NNE/Source/NNECore/Public/NNECoreTypes.h
+++ b/Engine/Plugins/Experimental/NNE/Source/NNECore/Public/NNECoreTypes.h
@@ -52,6 +52,11 @@ namespace UE::NNECore
 		TArray<int32, TInlineAllocator<MaxRank>> Data;
 
 	public:
+		/**
+		 * Default construct a symbolic tensor shape.
+		 */
+		FSymbolicTensorShape() = default;
+
 		/**
 		 * Construct this symbolic tensor shape from another one.
 		 * 
@@ -129,6 +134,10 @@ namespace UE::NNECore
 		TArray<uint32, TInlineAllocator<MaxRank>> Data;
 
 	public:
+		/**
+		 * Default construct a tensor shape.
+		 */
+		FTensorShape() = default;
 
 		/**
 		 * Construct this tensor shape from another one.
