NNE Only compile HLSL runtime shaders on supported platform.

What:
* Add data support flag for NNE shaders so platforms can opt in to the NNERuntimeHlsl shaders compilation.

Tests:
* NNE smoke test + style transfer demo on: win64 editor, win64 packaged

#rnx

diff --git a/Engine/Config/Mac/DataDrivenPlatformInfo.ini b/Engine/Config/Mac/DataDrivenPlatformInfo.ini
--- a/Engine/Config/Mac/DataDrivenPlatformInfo.ini
+++ b/Engine/Config/Mac/DataDrivenPlatformInfo.ini
@@ -60,6 +60,7 @@ bSupportsByteBufferComputeShaders = true
 bSupportsHairStrandGeometry = false
 bSupportsNanite=false
 bSupportsUInt64ImageAtomics=false
+bSupportsNNEShaders=true
 FriendlyName=LOCTEXT("FriendlyShaderPlatformName_MacSM5", "Mac Metal SM5")
 
 [ShaderPlatform METAL_MRT_MAC]
diff --git a/Engine/Config/VulkanPC/DataDrivenPlatformInfo.ini b/Engine/Config/VulkanPC/DataDrivenPlatformInfo.ini
--- a/Engine/Config/VulkanPC/DataDrivenPlatformInfo.ini
+++ b/Engine/Config/VulkanPC/DataDrivenPlatformInfo.ini
@@ -133,6 +133,7 @@ EnablesHLSL2021ByDefault=1
 FriendlyName=LOCTEXT("FriendlyShaderPlatformName_VulkanPC_SM5", "PC Vulkan SM5")
 
 BindlessSupport=Unsupported
+bSupportsNNEShaders=true
 
 
 
@@ -248,6 +249,7 @@ EnablesHLSL2021ByDefault=1
 FriendlyName=LOCTEXT("FriendlyShaderPlatformName_VulkanPC_SM6", "PC Vulkan SM6")
 
 BindlessSupport=AllShaderTypes
+bSupportsNNEShaders=true
 
 
 [ShaderPlatform VULKAN_PCES3_1]
diff --git a/Engine/Config/Windows/DataDrivenPlatformInfo.ini b/Engine/Config/Windows/DataDrivenPlatformInfo.ini
--- a/Engine/Config/Windows/DataDrivenPlatformInfo.ini
+++ b/Engine/Config/Windows/DataDrivenPlatformInfo.ini
@@ -76,6 +76,7 @@ BindlessSupport=RayTracingOnly
 SupportsMultiViewport=RuntimeDependent
 FriendlyName=LOCTEXT("FriendlyShaderPlatformName_D3D_SM5", "PC D3D SM5")
 bSupportsDualSourceBlending=true
+bSupportsNNEShaders=true
 
 [PreviewPlatform PCD3D_SM5]
 PlatformName=PCD3D_SM5
@@ -134,6 +135,7 @@ bSupportsRealTypes=RuntimeDependent
 EnablesHLSL2021ByDefault=1
 SupportsMultiViewport=RuntimeGuaranteed
 bSupportsDualSourceBlending=true
+bSupportsNNEShaders=true
 
 [PreviewPlatform PCD3D_SM6]
 PlatformName=PCD3D_SM6
diff --git a/Engine/Plugins/Experimental/NNE/Source/NNEHlslShaders/Internal/NNEHlslShadersBatchNormalizationCS.h b/Engine/Plugins/Experimental/NNE/Source/NNEHlslShaders/Internal/NNEHlslShadersBatchNormalizationCS.h
--- a/Engine/Plugins/Experimental/NNE/Source/NNEHlslShaders/Internal/NNEHlslShadersBatchNormalizationCS.h
+++ b/Engine/Plugins/Experimental/NNE/Source/NNEHlslShaders/Internal/NNEHlslShadersBatchNormalizationCS.h
@@ -2,6 +2,7 @@
 
 #pragma once
 
+#include "DataDrivenShaderPlatformInfo.h"
 #include "GlobalShader.h"
 #include "ShaderParameterUtils.h"
 #include "RenderGraphUtils.h"
@@ -38,5 +39,10 @@ namespace UE::NNEHlslShaders::Internal
 		END_SHADER_PARAMETER_STRUCT()
 
 		static void ModifyCompilationEnvironment(const FGlobalShaderPermutationParameters& InParameters, FShaderCompilerEnvironment& OutEnvironment);
+
+		static inline bool ShouldCompilePermutation(const FGlobalShaderPermutationParameters& Parameters)
+		{
+			return FDataDrivenShaderPlatformInfo::GetSupportsNNEShaders(Parameters.Platform);
+		}
 	};
 } // UE::NNEHlslShaders::Internal
diff --git a/Engine/Plugins/Experimental/NNE/Source/NNEHlslShaders/Internal/NNEHlslShadersConvCS.h b/Engine/Plugins/Experimental/NNE/Source/NNEHlslShaders/Internal/NNEHlslShadersConvCS.h
--- a/Engine/Plugins/Experimental/NNE/Source/NNEHlslShaders/Internal/NNEHlslShadersConvCS.h
+++ b/Engine/Plugins/Experimental/NNE/Source/NNEHlslShaders/Internal/NNEHlslShadersConvCS.h
@@ -2,6 +2,7 @@
 
 #pragma once
 
+#include "DataDrivenShaderPlatformInfo.h"
 #include "GlobalShader.h"
 #include "ShaderParameterUtils.h"
 #include "RenderGraphUtils.h"
@@ -111,5 +112,10 @@ namespace UE::NNEHlslShaders::Internal
 		static TArray<int32> GetPadding(TArrayView<const uint32> XShape, TArrayView<const uint32> WShape, EConvAutoPad AutoPad, TArrayView<const int32> Dilations, TArrayView<const int32> Strides, TArrayView<const int32> Pads);
 
 		static void LexFromString(EConvAutoPad& OutValue, const TCHAR* StringVal);
+
+		static inline bool ShouldCompilePermutation(const FGlobalShaderPermutationParameters& Parameters)
+		{
+			return FDataDrivenShaderPlatformInfo::GetSupportsNNEShaders(Parameters.Platform);
+		}
 	};
 } // UE::NNEHlslShaders::Internal
\ No newline at end of file
diff --git a/Engine/Plugins/Experimental/NNE/Source/NNEHlslShaders/Internal/NNEHlslShadersConvMatmulCS.h b/Engine/Plugins/Experimental/NNE/Source/NNEHlslShaders/Internal/NNEHlslShadersConvMatmulCS.h
--- a/Engine/Plugins/Experimental/NNE/Source/NNEHlslShaders/Internal/NNEHlslShadersConvMatmulCS.h
+++ b/Engine/Plugins/Experimental/NNE/Source/NNEHlslShaders/Internal/NNEHlslShadersConvMatmulCS.h
@@ -2,6 +2,7 @@
 
 #pragma once
 
+#include "DataDrivenShaderPlatformInfo.h"
 #include "GlobalShader.h"
 #include "ShaderParameterUtils.h"
 #include "RenderGraphUtils.h"
@@ -38,5 +39,10 @@ namespace UE::NNEHlslShaders::Internal
 		END_SHADER_PARAMETER_STRUCT()
 
 		static FIntVector GetGroupCount(TConstArrayView<uint32> OutputShape);
+
+		static inline bool ShouldCompilePermutation(const FGlobalShaderPermutationParameters& Parameters)
+		{
+			return FDataDrivenShaderPlatformInfo::GetSupportsNNEShaders(Parameters.Platform);
+		}
 	};
 } // UE::NNEHlslShaders::Internal
\ No newline at end of file
diff --git a/Engine/Plugins/Experimental/NNE/Source/NNEHlslShaders/Internal/NNEHlslShadersConvTransposeCS.h b/Engine/Plugins/Experimental/NNE/Source/NNEHlslShaders/Internal/NNEHlslShadersConvTransposeCS.h
--- a/Engine/Plugins/Experimental/NNE/Source/NNEHlslShaders/Internal/NNEHlslShadersConvTransposeCS.h
+++ b/Engine/Plugins/Experimental/NNE/Source/NNEHlslShaders/Internal/NNEHlslShadersConvTransposeCS.h
@@ -2,6 +2,7 @@
 
 #pragma once
 
+#include "DataDrivenShaderPlatformInfo.h"
 #include "GlobalShader.h"
 #include "ShaderParameterUtils.h"
 #include "RenderGraphUtils.h"
@@ -105,5 +106,10 @@ namespace UE::NNEHlslShaders::Internal
 		static int32 GetNumThreadsPerGroup(EConvTransposeGroupSize GroupSize);
 
 		static TArray<int32> GetGridShape(TArrayView<const int32> YShape, TArrayView<const int32> GroupShape);
+
+		static inline bool ShouldCompilePermutation(const FGlobalShaderPermutationParameters& Parameters)
+		{
+			return FDataDrivenShaderPlatformInfo::GetSupportsNNEShaders(Parameters.Platform);
+		}
 	};
 } // UE::NNEHlslShaders::Internal
\ No newline at end of file
diff --git a/Engine/Plugins/Experimental/NNE/Source/NNEHlslShaders/Internal/NNEHlslShadersElementWiseBinaryCS.h b/Engine/Plugins/Experimental/NNE/Source/NNEHlslShaders/Internal/NNEHlslShadersElementWiseBinaryCS.h
--- a/Engine/Plugins/Experimental/NNE/Source/NNEHlslShaders/Internal/NNEHlslShadersElementWiseBinaryCS.h
+++ b/Engine/Plugins/Experimental/NNE/Source/NNEHlslShaders/Internal/NNEHlslShadersElementWiseBinaryCS.h
@@ -2,6 +2,7 @@
 
 #pragma once
 
+#include "DataDrivenShaderPlatformInfo.h"
 #include "GlobalShader.h"
 #include "ShaderParameterUtils.h"
 #include "RenderGraphUtils.h"
@@ -39,6 +40,11 @@ namespace UE::NNEHlslShaders::Internal
 
 		static void ModifyCompilationEnvironment(const FGlobalShaderPermutationParameters& InParameters, FShaderCompilerEnvironment& OutEnvironment);
 
+		static inline bool ShouldCompilePermutation(const FGlobalShaderPermutationParameters& Parameters)
+		{
+			return FDataDrivenShaderPlatformInfo::GetSupportsNNEShaders(Parameters.Platform);
+		}
+
 	private:
 
 		static const FString GetOpFunc(NNECore::Internal::EElementWiseBinaryOperatorType OpType);
diff --git a/Engine/Plugins/Experimental/NNE/Source/NNEHlslShaders/Internal/NNEHlslShadersElementWiseUnaryCS.h b/Engine/Plugins/Experimental/NNE/Source/NNEHlslShaders/Internal/NNEHlslShadersElementWiseUnaryCS.h
--- a/Engine/Plugins/Experimental/NNE/Source/NNEHlslShaders/Internal/NNEHlslShadersElementWiseUnaryCS.h
+++ b/Engine/Plugins/Experimental/NNE/Source/NNEHlslShaders/Internal/NNEHlslShadersElementWiseUnaryCS.h
@@ -2,6 +2,7 @@
 
 #pragma once
 
+#include "DataDrivenShaderPlatformInfo.h"
 #include "GlobalShader.h"
 #include "ShaderParameterUtils.h"
 #include "RenderGraphUtils.h"
@@ -38,6 +39,11 @@ namespace UE::NNEHlslShaders::Internal
 
 		static void ModifyCompilationEnvironment(const FGlobalShaderPermutationParameters& InParameters, FShaderCompilerEnvironment& OutEnvironment);
 
+		static inline bool ShouldCompilePermutation(const FGlobalShaderPermutationParameters& Parameters)
+		{
+			return FDataDrivenShaderPlatformInfo::GetSupportsNNEShaders(Parameters.Platform);
+		}
+
 	private:
 
 		static const FString GetOpFunc(NNECore::Internal::EElementWiseUnaryOperatorType OpType);
diff --git a/Engine/Plugins/Experimental/NNE/Source/NNEHlslShaders/Internal/NNEHlslShadersElementWiseVariadicCS.h b/Engine/Plugins/Experimental/NNE/Source/NNEHlslShaders/Internal/NNEHlslShadersElementWiseVariadicCS.h
--- a/Engine/Plugins/Experimental/NNE/Source/NNEHlslShaders/Internal/NNEHlslShadersElementWiseVariadicCS.h
+++ b/Engine/Plugins/Experimental/NNE/Source/NNEHlslShaders/Internal/NNEHlslShadersElementWiseVariadicCS.h
@@ -2,6 +2,7 @@
 
 #pragma once
 
+#include "DataDrivenShaderPlatformInfo.h"
 #include "GlobalShader.h"
 #include "ShaderParameterUtils.h"
 #include "RenderGraphUtils.h"
@@ -47,6 +48,11 @@ namespace UE::NNEHlslShaders::Internal
 
 		static void ModifyCompilationEnvironment(const FGlobalShaderPermutationParameters& InParameters, FShaderCompilerEnvironment& OutEnvironment);
 
+		static inline bool ShouldCompilePermutation(const FGlobalShaderPermutationParameters& Parameters)
+		{
+			return FDataDrivenShaderPlatformInfo::GetSupportsNNEShaders(Parameters.Platform);
+		}
+
 	private:
 
 		static const FString GetOpFunc(NNECore::Internal::EElementWiseVariadicOperatorType OpType);
diff --git a/Engine/Plugins/Experimental/NNE/Source/NNEHlslShaders/Internal/NNEHlslShadersGatherCS.h b/Engine/Plugins/Experimental/NNE/Source/NNEHlslShaders/Internal/NNEHlslShadersGatherCS.h
--- a/Engine/Plugins/Experimental/NNE/Source/NNEHlslShaders/Internal/NNEHlslShadersGatherCS.h
+++ b/Engine/Plugins/Experimental/NNE/Source/NNEHlslShaders/Internal/NNEHlslShadersGatherCS.h
@@ -2,6 +2,7 @@
 
 #pragma once
 
+#include "DataDrivenShaderPlatformInfo.h"
 #include "GlobalShader.h"
 #include "ShaderParameterUtils.h"
 #include "RenderGraphUtils.h"
@@ -43,6 +44,11 @@ namespace UE::NNEHlslShaders::Internal
 
 		static void ModifyCompilationEnvironment(const FGlobalShaderPermutationParameters& InParameters, FShaderCompilerEnvironment& OutEnvironment);
 
+		static inline bool ShouldCompilePermutation(const FGlobalShaderPermutationParameters& Parameters)
+		{
+			return FDataDrivenShaderPlatformInfo::GetSupportsNNEShaders(Parameters.Platform);
+		}
+
 		static void FillInParameters(int32 Axis, const NNECore::Internal::FTensor& Data, const NNECore::Internal::FTensor& Indices, FParameters& Parameters);
 
 		static FIntVector GetGroupCount(const FParameters& Parameters);
diff --git a/Engine/Plugins/Experimental/NNE/Source/NNEHlslShaders/Internal/NNEHlslShadersGemmCS.h b/Engine/Plugins/Experimental/NNE/Source/NNEHlslShaders/Internal/NNEHlslShadersGemmCS.h
--- a/Engine/Plugins/Experimental/NNE/Source/NNEHlslShaders/Internal/NNEHlslShadersGemmCS.h
+++ b/Engine/Plugins/Experimental/NNE/Source/NNEHlslShaders/Internal/NNEHlslShadersGemmCS.h
@@ -2,6 +2,7 @@
 
 #pragma once
 
+#include "DataDrivenShaderPlatformInfo.h"
 #include "GlobalShader.h"
 #include "ShaderParameterUtils.h"
 #include "RenderGraphUtils.h"
@@ -80,6 +81,11 @@ namespace UE::NNEHlslShaders::Internal
 
 		static void ModifyCompilationEnvironment(const FGlobalShaderPermutationParameters& InParameters, FShaderCompilerEnvironment& OutEnvironment);
 
+		static inline bool ShouldCompilePermutation(const FGlobalShaderPermutationParameters& Parameters)
+		{
+			return FDataDrivenShaderPlatformInfo::GetSupportsNNEShaders(Parameters.Platform);
+		}
+
 		static void FillInParameters(float Alpha, float Beta, int32 TransA, int32 TransB, const NNECore::Internal::FTensor& InputA, const NNECore::Internal::FTensor& InputB, const NNECore::Internal::FTensor* InputC, float CScalar, FParameters& Parameters);
 
 		static void FillInParametersMatMul(const NNECore::Internal::FTensor& InputA, const NNECore::Internal::FTensor& InputB, FParameters& Parameters);
diff --git a/Engine/Plugins/Experimental/NNE/Source/NNEHlslShaders/Internal/NNEHlslShadersInstanceNormalizationCS.h b/Engine/Plugins/Experimental/NNE/Source/NNEHlslShaders/Internal/NNEHlslShadersInstanceNormalizationCS.h
--- a/Engine/Plugins/Experimental/NNE/Source/NNEHlslShaders/Internal/NNEHlslShadersInstanceNormalizationCS.h
+++ b/Engine/Plugins/Experimental/NNE/Source/NNEHlslShaders/Internal/NNEHlslShadersInstanceNormalizationCS.h
@@ -2,6 +2,7 @@
 
 #pragma once
 
+#include "DataDrivenShaderPlatformInfo.h"
 #include "GlobalShader.h"
 #include "ShaderParameterUtils.h"
 #include "RenderGraphUtils.h"
@@ -47,5 +48,10 @@ namespace UE::NNEHlslShaders::Internal
 		static EInstanceNormalizationAlgorithm GetAlgorithm(const FParameters& Parameters);
 
 		static void LexFromString(EInstanceNormalizationAlgorithm& OutValue, const TCHAR* StringVal);
+
+		static inline bool ShouldCompilePermutation(const FGlobalShaderPermutationParameters& Parameters)
+		{
+			return FDataDrivenShaderPlatformInfo::GetSupportsNNEShaders(Parameters.Platform);
+		}
 	};
 } // UE::NNEHlslShaders::Internal
\ No newline at end of file
diff --git a/Engine/Plugins/Experimental/NNE/Source/NNEHlslShaders/Internal/NNEHlslShadersNeuralPostProcessingCS.h b/Engine/Plugins/Experimental/NNE/Source/NNEHlslShaders/Internal/NNEHlslShadersNeuralPostProcessingCS.h
--- a/Engine/Plugins/Experimental/NNE/Source/NNEHlslShaders/Internal/NNEHlslShadersNeuralPostProcessingCS.h
+++ b/Engine/Plugins/Experimental/NNE/Source/NNEHlslShaders/Internal/NNEHlslShadersNeuralPostProcessingCS.h
@@ -2,6 +2,7 @@
 
 #pragma once
 
+#include "DataDrivenShaderPlatformInfo.h"
 #include "GlobalShader.h"
 #include "ShaderParameterUtils.h"
 #include "RenderGraphUtils.h"
@@ -45,6 +46,11 @@ namespace UE::NNEHlslShaders::Internal
 		END_SHADER_PARAMETER_STRUCT()
 
 		static void ModifyCompilationEnvironment(const FGlobalShaderPermutationParameters& InParameters, FShaderCompilerEnvironment& OutEnvironment);
+
+		static inline bool ShouldCompilePermutation(const FGlobalShaderPermutationParameters& Parameters)
+		{
+			return FDataDrivenShaderPlatformInfo::GetSupportsNNEShaders(Parameters.Platform);
+		}
 	};
 
 	class NNEHLSLSHADERS_API TNeuralPostProcessingPreStepCS : public FGlobalShader
@@ -66,6 +72,11 @@ namespace UE::NNEHlslShaders::Internal
 		END_SHADER_PARAMETER_STRUCT()
 
 		static void ModifyCompilationEnvironment(const FGlobalShaderPermutationParameters& InParameters, FShaderCompilerEnvironment& OutEnvironment);
+
+		static inline bool ShouldCompilePermutation(const FGlobalShaderPermutationParameters& Parameters)
+		{
+			return FDataDrivenShaderPlatformInfo::GetSupportsNNEShaders(Parameters.Platform);
+		}
 	};
 
 	class NNEHLSLSHADERS_API TNeuralPostProcessingPostStepCS : public FGlobalShader
@@ -91,6 +102,11 @@ namespace UE::NNEHlslShaders::Internal
 		END_SHADER_PARAMETER_STRUCT()
 
 		static void ModifyCompilationEnvironment(const FGlobalShaderPermutationParameters& InParameters, FShaderCompilerEnvironment& OutEnvironment);
+
+		static inline bool ShouldCompilePermutation(const FGlobalShaderPermutationParameters& Parameters)
+		{
+			return FDataDrivenShaderPlatformInfo::GetSupportsNNEShaders(Parameters.Platform);
+		}
 	};
 
 	class NNEHLSLSHADERS_API TNeuralPostProcessingWriteOutputPS : public FGlobalShader
@@ -108,5 +124,10 @@ namespace UE::NNEHlslShaders::Internal
 		END_SHADER_PARAMETER_STRUCT()
 
 		static void ModifyCompilationEnvironment(const FGlobalShaderPermutationParameters& InParameters, FShaderCompilerEnvironment& OutEnvironment);
+
+		static inline bool ShouldCompilePermutation(const FGlobalShaderPermutationParameters& Parameters)
+		{
+			return FDataDrivenShaderPlatformInfo::GetSupportsNNEShaders(Parameters.Platform);
+		}
 	};
 } // UE::NNEHlslShaders::Internal
\ No newline at end of file
diff --git a/Engine/Plugins/Experimental/NNE/Source/NNEHlslShaders/Internal/NNEHlslShadersPadCS.h b/Engine/Plugins/Experimental/NNE/Source/NNEHlslShaders/Internal/NNEHlslShadersPadCS.h
--- a/Engine/Plugins/Experimental/NNE/Source/NNEHlslShaders/Internal/NNEHlslShadersPadCS.h
+++ b/Engine/Plugins/Experimental/NNE/Source/NNEHlslShaders/Internal/NNEHlslShadersPadCS.h
@@ -2,6 +2,7 @@
 
 #pragma once
 
+#include "DataDrivenShaderPlatformInfo.h"
 #include "GlobalShader.h"
 #include "ShaderParameterUtils.h"
 #include "RenderGraphUtils.h"
@@ -45,6 +46,12 @@ namespace UE::NNEHlslShaders::Internal
 		END_SHADER_PARAMETER_STRUCT()
 
 		static void ModifyCompilationEnvironment(const FGlobalShaderPermutationParameters& InParameters, FShaderCompilerEnvironment& OutEnvironment);
+
+		static inline bool ShouldCompilePermutation(const FGlobalShaderPermutationParameters& Parameters)
+		{
+			return FDataDrivenShaderPlatformInfo::GetSupportsNNEShaders(Parameters.Platform);
+		}
+
 		static void LexFromString(EPadMode& OutValue, const TCHAR* StringVal);
 	};
 } // UE::NNEHlslShaders::Internal
diff --git a/Engine/Plugins/Experimental/NNE/Source/NNEHlslShaders/Internal/NNEHlslShadersUpsampleCS.h b/Engine/Plugins/Experimental/NNE/Source/NNEHlslShaders/Internal/NNEHlslShadersUpsampleCS.h
--- a/Engine/Plugins/Experimental/NNE/Source/NNEHlslShaders/Internal/NNEHlslShadersUpsampleCS.h
+++ b/Engine/Plugins/Experimental/NNE/Source/NNEHlslShaders/Internal/NNEHlslShadersUpsampleCS.h
@@ -2,6 +2,7 @@
 
 #pragma once
 
+#include "DataDrivenShaderPlatformInfo.h"
 #include "GlobalShader.h"
 #include "ShaderParameterUtils.h"
 #include "RenderGraphUtils.h"
@@ -35,5 +36,10 @@ namespace UE::NNEHlslShaders::Internal
 		END_SHADER_PARAMETER_STRUCT()
 
 		static void ModifyCompilationEnvironment(const FGlobalShaderPermutationParameters& InParameters, FShaderCompilerEnvironment& OutEnvironment);
+
+		static inline bool ShouldCompilePermutation(const FGlobalShaderPermutationParameters& Parameters)
+		{
+			return FDataDrivenShaderPlatformInfo::GetSupportsNNEShaders(Parameters.Platform);
+		}
 	};
 } // UE::NNEHlslShaders::Internal
diff --git a/Engine/Source/Runtime/RHI/Private/DataDrivenShaderPlatformInfo.cpp b/Engine/Source/Runtime/RHI/Private/DataDrivenShaderPlatformInfo.cpp
--- a/Engine/Source/Runtime/RHI/Private/DataDrivenShaderPlatformInfo.cpp
+++ b/Engine/Source/Runtime/RHI/Private/DataDrivenShaderPlatformInfo.cpp
@@ -269,6 +269,7 @@ void FGenericDataDrivenShaderPlatformInfo::ParseDataDrivenShaderInfo(const FConf
 	GET_SECTION_BOOL_HELPER(bSupportsSceneDataCompressedTransforms);
 	GET_SECTION_BOOL_HELPER(bSupportsSwapchainUAVs);
 	GET_SECTION_BOOL_HELPER(bSupportsClipDistance);
+	GET_SECTION_BOOL_HELPER(bSupportsNNEShaders);
 #undef GET_SECTION_BOOL_HELPER
 #undef GET_SECTION_INT_HELPER
 #undef GET_SECTION_SUPPORT_HELPER
diff --git a/Engine/Source/Runtime/RHI/Public/DataDrivenShaderPlatformInfo.h b/Engine/Source/Runtime/RHI/Public/DataDrivenShaderPlatformInfo.h
--- a/Engine/Source/Runtime/RHI/Public/DataDrivenShaderPlatformInfo.h
+++ b/Engine/Source/Runtime/RHI/Public/DataDrivenShaderPlatformInfo.h
@@ -110,6 +110,7 @@ class FGenericDataDrivenShaderPlatformInfo
 	uint32 bIsPreviewPlatform : 1;
 	uint32 bSupportsSwapchainUAVs : 1;
 	uint32 bSupportsClipDistance : 1;
+	uint32 bSupportsNNEShaders: 1;
 
 #if WITH_EDITOR
 	FText FriendlyName;
@@ -728,6 +729,12 @@ public:
 		return Infos[Platform].bSupportsClipDistance;
 	}
 
+	static FORCEINLINE_DEBUGGABLE const bool GetSupportsNNEShaders(const FStaticShaderPlatform Platform)
+	{
+		check(IsValid(Platform));
+		return Infos[Platform].bSupportsNNEShaders;
+	}
+
 #if WITH_EDITOR
 	RHI_API static FText GetFriendlyName(const FStaticShaderPlatform Platform);
 #endif
