NNE Cleanup: factorize HLSL shader compilation toggle following review

Tests: NNE smoke test

#rnx

diff --git a/Engine/Plugins/Experimental/NNE/Source/NNEHlslShaders/Internal/NNEHlslShaderBase.h b/Engine/Plugins/Experimental/NNE/Source/NNEHlslShaders/Internal/NNEHlslShaderBase.h
new file mode 100644
--- /dev/null
+++ b/Engine/Plugins/Experimental/NNE/Source/NNEHlslShaders/Internal/NNEHlslShaderBase.h
@@ -0,0 +1,19 @@
+// Copyright Epic Games, Inc. All Rights Reserved.
+
+#pragma once
+
+#include "GlobalShader.h"
+
+namespace UE::NNEHlslShaders::Internal
+{
+	class FHlslShaderBase : public FGlobalShader
+	{
+	public:
+
+		static bool ShouldCompilePermutation(const FGlobalShaderPermutationParameters& Parameters);
+
+		FHlslShaderBase() {}
+		FHlslShaderBase(const ShaderMetaType::CompiledShaderInitializerType& initType) : FGlobalShader(initType) {}
+	};
+	
+} // UE::NNEHlslShaders::Internal
\ No newline at end of file
diff --git a/Engine/Plugins/Experimental/NNE/Source/NNEHlslShaders/Internal/NNEHlslShadersBatchNormalizationCS.h b/Engine/Plugins/Experimental/NNE/Source/NNEHlslShaders/Internal/NNEHlslShadersBatchNormalizationCS.h
--- a/Engine/Plugins/Experimental/NNE/Source/NNEHlslShaders/Internal/NNEHlslShadersBatchNormalizationCS.h
+++ b/Engine/Plugins/Experimental/NNE/Source/NNEHlslShaders/Internal/NNEHlslShadersBatchNormalizationCS.h
@@ -2,11 +2,11 @@
 
 #pragma once
 
-#include "DataDrivenShaderPlatformInfo.h"
-#include "GlobalShader.h"
-#include "ShaderParameterUtils.h"
-#include "RenderGraphUtils.h"
 #include "NNECoreOperator.h"
+#include "NNEHlslShaderBase.h"
+#include "RenderGraphUtils.h"
+#include "ShaderParameterUtils.h"
+
 
 namespace UE::NNEHlslShaders::Internal
 {
@@ -17,10 +17,10 @@ namespace UE::NNEHlslShaders::Internal
 		static const int32 NUM_GROUP_THREADS{ 256 };
 	};
 
-	class NNEHLSLSHADERS_API TBatchNormalizationCS : public FGlobalShader
+	class NNEHLSLSHADERS_API TBatchNormalizationCS : public FHlslShaderBase
 	{
 		DECLARE_GLOBAL_SHADER(TBatchNormalizationCS);
-		SHADER_USE_PARAMETER_STRUCT(TBatchNormalizationCS, FGlobalShader)
+		SHADER_USE_PARAMETER_STRUCT(TBatchNormalizationCS, FHlslShaderBase)
 
 	public:
 
@@ -39,10 +39,5 @@ namespace UE::NNEHlslShaders::Internal
 		END_SHADER_PARAMETER_STRUCT()
 
 		static void ModifyCompilationEnvironment(const FGlobalShaderPermutationParameters& InParameters, FShaderCompilerEnvironment& OutEnvironment);
-
-		static inline bool ShouldCompilePermutation(const FGlobalShaderPermutationParameters& Parameters)
-		{
-			return FDataDrivenShaderPlatformInfo::GetSupportsNNEShaders(Parameters.Platform);
-		}
 	};
 } // UE::NNEHlslShaders::Internal
diff --git a/Engine/Plugins/Experimental/NNE/Source/NNEHlslShaders/Internal/NNEHlslShadersConvCS.h b/Engine/Plugins/Experimental/NNE/Source/NNEHlslShaders/Internal/NNEHlslShadersConvCS.h
--- a/Engine/Plugins/Experimental/NNE/Source/NNEHlslShaders/Internal/NNEHlslShadersConvCS.h
+++ b/Engine/Plugins/Experimental/NNE/Source/NNEHlslShaders/Internal/NNEHlslShadersConvCS.h
@@ -2,10 +2,9 @@
 
 #pragma once
 
-#include "DataDrivenShaderPlatformInfo.h"
-#include "GlobalShader.h"
-#include "ShaderParameterUtils.h"
+#include "NNEHlslShaderBase.h"
 #include "RenderGraphUtils.h"
+#include "ShaderParameterUtils.h"
 
 namespace UE::NNEHlslShaders::Internal
 {
@@ -40,10 +39,10 @@ namespace UE::NNEHlslShaders::Internal
 		static const int32 MAX_NUM_READS_PER_THREAD_POW2{3};
 	};
 
-	class NNEHLSLSHADERS_API FConvCS : public FGlobalShader
+	class NNEHLSLSHADERS_API FConvCS : public FHlslShaderBase
 	{
 		DECLARE_GLOBAL_SHADER(FConvCS);
-		SHADER_USE_PARAMETER_STRUCT(FConvCS, FGlobalShader)
+		SHADER_USE_PARAMETER_STRUCT(FConvCS, FHlslShaderBase)
 
 		class FConvAlgorithm : SHADER_PERMUTATION_ENUM_CLASS("ALGORITHM", EConvAlgorithm);
 		class FConvAreWeightsTransposed : SHADER_PERMUTATION_BOOL("WEIGHTS_TRANSPOSED");
@@ -112,10 +111,5 @@ namespace UE::NNEHlslShaders::Internal
 		static TArray<int32> GetPadding(TArrayView<const uint32> XShape, TArrayView<const uint32> WShape, EConvAutoPad AutoPad, TArrayView<const int32> Dilations, TArrayView<const int32> Strides, TArrayView<const int32> Pads);
 
 		static void LexFromString(EConvAutoPad& OutValue, const TCHAR* StringVal);
-
-		static inline bool ShouldCompilePermutation(const FGlobalShaderPermutationParameters& Parameters)
-		{
-			return FDataDrivenShaderPlatformInfo::GetSupportsNNEShaders(Parameters.Platform);
-		}
 	};
 } // UE::NNEHlslShaders::Internal
\ No newline at end of file
diff --git a/Engine/Plugins/Experimental/NNE/Source/NNEHlslShaders/Internal/NNEHlslShadersConvMatmulCS.h b/Engine/Plugins/Experimental/NNE/Source/NNEHlslShaders/Internal/NNEHlslShadersConvMatmulCS.h
--- a/Engine/Plugins/Experimental/NNE/Source/NNEHlslShaders/Internal/NNEHlslShadersConvMatmulCS.h
+++ b/Engine/Plugins/Experimental/NNE/Source/NNEHlslShaders/Internal/NNEHlslShadersConvMatmulCS.h
@@ -2,17 +2,16 @@
 
 #pragma once
 
-#include "DataDrivenShaderPlatformInfo.h"
-#include "GlobalShader.h"
-#include "ShaderParameterUtils.h"
+#include "NNEHlslShaderBase.h"
 #include "RenderGraphUtils.h"
+#include "ShaderParameterUtils.h"
 
 namespace UE::NNEHlslShaders::Internal
 {
-	class NNEHLSLSHADERS_API FConvMatmulCS : public FGlobalShader
+	class NNEHLSLSHADERS_API FConvMatmulCS : public FHlslShaderBase
 	{
 		DECLARE_GLOBAL_SHADER(FConvMatmulCS);
-		SHADER_USE_PARAMETER_STRUCT(FConvMatmulCS, FGlobalShader)
+		SHADER_USE_PARAMETER_STRUCT(FConvMatmulCS, FHlslShaderBase)
 
 		class FConvMatmulAreWeightsTransposed : SHADER_PERMUTATION_BOOL("WEIGHTS_TRANSPOSED");
 		class FConvMatmulHasBias : SHADER_PERMUTATION_BOOL("HAS_BIAS");
@@ -39,10 +38,5 @@ namespace UE::NNEHlslShaders::Internal
 		END_SHADER_PARAMETER_STRUCT()
 
 		static FIntVector GetGroupCount(TConstArrayView<uint32> OutputShape);
-
-		static inline bool ShouldCompilePermutation(const FGlobalShaderPermutationParameters& Parameters)
-		{
-			return FDataDrivenShaderPlatformInfo::GetSupportsNNEShaders(Parameters.Platform);
-		}
 	};
 } // UE::NNEHlslShaders::Internal
\ No newline at end of file
diff --git a/Engine/Plugins/Experimental/NNE/Source/NNEHlslShaders/Internal/NNEHlslShadersConvTransposeCS.h b/Engine/Plugins/Experimental/NNE/Source/NNEHlslShaders/Internal/NNEHlslShadersConvTransposeCS.h
--- a/Engine/Plugins/Experimental/NNE/Source/NNEHlslShaders/Internal/NNEHlslShadersConvTransposeCS.h
+++ b/Engine/Plugins/Experimental/NNE/Source/NNEHlslShaders/Internal/NNEHlslShadersConvTransposeCS.h
@@ -2,10 +2,9 @@
 
 #pragma once
 
-#include "DataDrivenShaderPlatformInfo.h"
-#include "GlobalShader.h"
-#include "ShaderParameterUtils.h"
+#include "NNEHlslShaderBase.h"
 #include "RenderGraphUtils.h"
+#include "ShaderParameterUtils.h"
 
 namespace UE::NNEHlslShaders::Internal
 {
@@ -41,10 +40,10 @@ namespace UE::NNEHlslShaders::Internal
 		static const int32 MAX_NUM_READS_PER_THREAD_POW2{3};
 	};
 
-	class NNEHLSLSHADERS_API FConvTransposeCS : public FGlobalShader
+	class NNEHLSLSHADERS_API FConvTransposeCS : public FHlslShaderBase
 	{
 		DECLARE_GLOBAL_SHADER(FConvTransposeCS);
-		SHADER_USE_PARAMETER_STRUCT(FConvTransposeCS, FGlobalShader);
+		SHADER_USE_PARAMETER_STRUCT(FConvTransposeCS, FHlslShaderBase);
 
 		class FConvTransposeAlgorithm : SHADER_PERMUTATION_ENUM_CLASS("ALGORITHM", EConvTransposeAlgorithm);
 		class FConvTransposeGroupSize : SHADER_PERMUTATION_ENUM_CLASS("GROUP_SIZE", EConvTransposeGroupSize);
@@ -106,10 +105,5 @@ namespace UE::NNEHlslShaders::Internal
 		static int32 GetNumThreadsPerGroup(EConvTransposeGroupSize GroupSize);
 
 		static TArray<int32> GetGridShape(TArrayView<const int32> YShape, TArrayView<const int32> GroupShape);
-
-		static inline bool ShouldCompilePermutation(const FGlobalShaderPermutationParameters& Parameters)
-		{
-			return FDataDrivenShaderPlatformInfo::GetSupportsNNEShaders(Parameters.Platform);
-		}
 	};
 } // UE::NNEHlslShaders::Internal
\ No newline at end of file
diff --git a/Engine/Plugins/Experimental/NNE/Source/NNEHlslShaders/Internal/NNEHlslShadersElementWiseBinaryCS.h b/Engine/Plugins/Experimental/NNE/Source/NNEHlslShaders/Internal/NNEHlslShadersElementWiseBinaryCS.h
--- a/Engine/Plugins/Experimental/NNE/Source/NNEHlslShaders/Internal/NNEHlslShadersElementWiseBinaryCS.h
+++ b/Engine/Plugins/Experimental/NNE/Source/NNEHlslShaders/Internal/NNEHlslShadersElementWiseBinaryCS.h
@@ -2,11 +2,10 @@
 
 #pragma once
 
-#include "DataDrivenShaderPlatformInfo.h"
-#include "GlobalShader.h"
-#include "ShaderParameterUtils.h"
-#include "RenderGraphUtils.h"
 #include "NNECoreOperator.h"
+#include "NNEHlslShaderBase.h"
+#include "RenderGraphUtils.h"
+#include "ShaderParameterUtils.h"
 
 namespace UE::NNEHlslShaders::Internal
 {
@@ -18,10 +17,10 @@ namespace UE::NNEHlslShaders::Internal
 		static const int32 NUM_GROUP_THREADS{ 256 };
 	};
 
-	class NNEHLSLSHADERS_API TElementWiseBinaryCS : public FGlobalShader
+	class NNEHLSLSHADERS_API TElementWiseBinaryCS : public FHlslShaderBase
 	{
 		DECLARE_GLOBAL_SHADER(TElementWiseBinaryCS);
-		SHADER_USE_PARAMETER_STRUCT(TElementWiseBinaryCS, FGlobalShader)
+		SHADER_USE_PARAMETER_STRUCT(TElementWiseBinaryCS, FHlslShaderBase)
 
 		class FOperatorType : SHADER_PERMUTATION_ENUM_CLASS("OP_TYPENAME", NNECore::Internal::EElementWiseBinaryOperatorType);
 		class FBinaryNumDimensions : SHADER_PERMUTATION_RANGE_INT("NUM_DIMENSIONS", 1, FElementWiseBinaryConstants::MAX_NUM_DIMENSIONS);
@@ -40,11 +39,6 @@ namespace UE::NNEHlslShaders::Internal
 
 		static void ModifyCompilationEnvironment(const FGlobalShaderPermutationParameters& InParameters, FShaderCompilerEnvironment& OutEnvironment);
 
-		static inline bool ShouldCompilePermutation(const FGlobalShaderPermutationParameters& Parameters)
-		{
-			return FDataDrivenShaderPlatformInfo::GetSupportsNNEShaders(Parameters.Platform);
-		}
-
 	private:
 
 		static const FString GetOpFunc(NNECore::Internal::EElementWiseBinaryOperatorType OpType);
diff --git a/Engine/Plugins/Experimental/NNE/Source/NNEHlslShaders/Internal/NNEHlslShadersElementWiseUnaryCS.h b/Engine/Plugins/Experimental/NNE/Source/NNEHlslShaders/Internal/NNEHlslShadersElementWiseUnaryCS.h
--- a/Engine/Plugins/Experimental/NNE/Source/NNEHlslShaders/Internal/NNEHlslShadersElementWiseUnaryCS.h
+++ b/Engine/Plugins/Experimental/NNE/Source/NNEHlslShaders/Internal/NNEHlslShadersElementWiseUnaryCS.h
@@ -2,11 +2,10 @@
 
 #pragma once
 
-#include "DataDrivenShaderPlatformInfo.h"
-#include "GlobalShader.h"
-#include "ShaderParameterUtils.h"
-#include "RenderGraphUtils.h"
 #include "NNECoreOperator.h"
+#include "NNEHlslShaderBase.h"
+#include "RenderGraphUtils.h"
+#include "ShaderParameterUtils.h"
 
 namespace UE::NNEHlslShaders::Internal
 {
@@ -17,10 +16,10 @@ namespace UE::NNEHlslShaders::Internal
 		static const int32 NUM_GROUP_THREADS{ 256 };
 	};
 
-	class NNEHLSLSHADERS_API TElementWiseUnaryCS : public FGlobalShader
+	class NNEHLSLSHADERS_API TElementWiseUnaryCS : public FHlslShaderBase
 	{
 		DECLARE_GLOBAL_SHADER(TElementWiseUnaryCS);
-		SHADER_USE_PARAMETER_STRUCT(TElementWiseUnaryCS, FGlobalShader)
+		SHADER_USE_PARAMETER_STRUCT(TElementWiseUnaryCS, FHlslShaderBase)
 
 		class FOperatorType : SHADER_PERMUTATION_ENUM_CLASS("OP_TYPENAME", NNECore::Internal::EElementWiseUnaryOperatorType);
 		using FPermutationDomain = TShaderPermutationDomain<FOperatorType>;
@@ -39,11 +38,6 @@ namespace UE::NNEHlslShaders::Internal
 
 		static void ModifyCompilationEnvironment(const FGlobalShaderPermutationParameters& InParameters, FShaderCompilerEnvironment& OutEnvironment);
 
-		static inline bool ShouldCompilePermutation(const FGlobalShaderPermutationParameters& Parameters)
-		{
-			return FDataDrivenShaderPlatformInfo::GetSupportsNNEShaders(Parameters.Platform);
-		}
-
 	private:
 
 		static const FString GetOpFunc(NNECore::Internal::EElementWiseUnaryOperatorType OpType);
diff --git a/Engine/Plugins/Experimental/NNE/Source/NNEHlslShaders/Internal/NNEHlslShadersElementWiseVariadicCS.h b/Engine/Plugins/Experimental/NNE/Source/NNEHlslShaders/Internal/NNEHlslShadersElementWiseVariadicCS.h
--- a/Engine/Plugins/Experimental/NNE/Source/NNEHlslShaders/Internal/NNEHlslShadersElementWiseVariadicCS.h
+++ b/Engine/Plugins/Experimental/NNE/Source/NNEHlslShaders/Internal/NNEHlslShadersElementWiseVariadicCS.h
@@ -2,11 +2,10 @@
 
 #pragma once
 
-#include "DataDrivenShaderPlatformInfo.h"
-#include "GlobalShader.h"
-#include "ShaderParameterUtils.h"
-#include "RenderGraphUtils.h"
 #include "NNECoreOperator.h"
+#include "NNEHlslShaderBase.h"
+#include "RenderGraphUtils.h"
+#include "ShaderParameterUtils.h"
 
 namespace UE::NNEHlslShaders::Internal
 {
@@ -19,10 +18,10 @@ namespace UE::NNEHlslShaders::Internal
 		static const uint32 MAX_NUM_INPUT{ 4 };
 	};
 
-	class NNEHLSLSHADERS_API TElementWiseVariadicCS : public FGlobalShader
+	class NNEHLSLSHADERS_API TElementWiseVariadicCS : public FHlslShaderBase
 	{
 		DECLARE_GLOBAL_SHADER(TElementWiseVariadicCS);
-		SHADER_USE_PARAMETER_STRUCT(TElementWiseVariadicCS, FGlobalShader)
+		SHADER_USE_PARAMETER_STRUCT(TElementWiseVariadicCS, FHlslShaderBase)
 
 		class FOperatorType : SHADER_PERMUTATION_ENUM_CLASS("OP_TYPENAME", NNECore::Internal::EElementWiseVariadicOperatorType);
 		class FApplyScale : SHADER_PERMUTATION_BOOL("APPLYSCALE");
@@ -48,11 +47,6 @@ namespace UE::NNEHlslShaders::Internal
 
 		static void ModifyCompilationEnvironment(const FGlobalShaderPermutationParameters& InParameters, FShaderCompilerEnvironment& OutEnvironment);
 
-		static inline bool ShouldCompilePermutation(const FGlobalShaderPermutationParameters& Parameters)
-		{
-			return FDataDrivenShaderPlatformInfo::GetSupportsNNEShaders(Parameters.Platform);
-		}
-
 	private:
 
 		static const FString GetOpFunc(NNECore::Internal::EElementWiseVariadicOperatorType OpType);
diff --git a/Engine/Plugins/Experimental/NNE/Source/NNEHlslShaders/Internal/NNEHlslShadersGatherCS.h b/Engine/Plugins/Experimental/NNE/Source/NNEHlslShaders/Internal/NNEHlslShadersGatherCS.h
--- a/Engine/Plugins/Experimental/NNE/Source/NNEHlslShaders/Internal/NNEHlslShadersGatherCS.h
+++ b/Engine/Plugins/Experimental/NNE/Source/NNEHlslShaders/Internal/NNEHlslShadersGatherCS.h
@@ -2,10 +2,9 @@
 
 #pragma once
 
-#include "DataDrivenShaderPlatformInfo.h"
-#include "GlobalShader.h"
-#include "ShaderParameterUtils.h"
+#include "NNEHlslShaderBase.h"
 #include "RenderGraphUtils.h"
+#include "ShaderParameterUtils.h"
 
 namespace UE::NNECore::Internal { class FTensor; }
 
@@ -20,10 +19,10 @@ namespace UE::NNEHlslShaders::Internal
 	};
 
 	// template <typename DataElementType, typename IndicesElementType>
-	class NNEHLSLSHADERS_API TGatherCS : public FGlobalShader
+	class NNEHLSLSHADERS_API TGatherCS : public FHlslShaderBase
 	{
 		DECLARE_GLOBAL_SHADER(TGatherCS);
-		SHADER_USE_PARAMETER_STRUCT(TGatherCS, FGlobalShader)
+		SHADER_USE_PARAMETER_STRUCT(TGatherCS, FHlslShaderBase)
 
 		class FGatherNumOutputDimensions : SHADER_PERMUTATION_RANGE_INT("NUM_OUTPUT_DIMENSIONS", 1, FGatherConstants::MAX_NUM_DIMENSIONS);
 		using FPermutationDomain = TShaderPermutationDomain<FGatherNumOutputDimensions>;
@@ -44,11 +43,6 @@ namespace UE::NNEHlslShaders::Internal
 
 		static void ModifyCompilationEnvironment(const FGlobalShaderPermutationParameters& InParameters, FShaderCompilerEnvironment& OutEnvironment);
 
-		static inline bool ShouldCompilePermutation(const FGlobalShaderPermutationParameters& Parameters)
-		{
-			return FDataDrivenShaderPlatformInfo::GetSupportsNNEShaders(Parameters.Platform);
-		}
-
 		static void FillInParameters(int32 Axis, const NNECore::Internal::FTensor& Data, const NNECore::Internal::FTensor& Indices, FParameters& Parameters);
 
 		static FIntVector GetGroupCount(const FParameters& Parameters);
diff --git a/Engine/Plugins/Experimental/NNE/Source/NNEHlslShaders/Internal/NNEHlslShadersGemmCS.h b/Engine/Plugins/Experimental/NNE/Source/NNEHlslShaders/Internal/NNEHlslShadersGemmCS.h
--- a/Engine/Plugins/Experimental/NNE/Source/NNEHlslShaders/Internal/NNEHlslShadersGemmCS.h
+++ b/Engine/Plugins/Experimental/NNE/Source/NNEHlslShaders/Internal/NNEHlslShadersGemmCS.h
@@ -2,10 +2,9 @@
 
 #pragma once
 
-#include "DataDrivenShaderPlatformInfo.h"
-#include "GlobalShader.h"
-#include "ShaderParameterUtils.h"
+#include "NNEHlslShaderBase.h"
 #include "RenderGraphUtils.h"
+#include "ShaderParameterUtils.h"
 
 namespace UE::NNECore::Internal { class FTensor; }
 
@@ -46,10 +45,10 @@ namespace UE::NNEHlslShaders::Internal
 		static const int32 MAX_NUM_STACK_DIMENSIONS{8};
 	};
 
-	class NNEHLSLSHADERS_API TGemmCS : public FGlobalShader
+	class NNEHLSLSHADERS_API TGemmCS : public FHlslShaderBase
 	{
 		DECLARE_GLOBAL_SHADER(TGemmCS);
-		SHADER_USE_PARAMETER_STRUCT(TGemmCS, FGlobalShader)
+		SHADER_USE_PARAMETER_STRUCT(TGemmCS, FHlslShaderBase)
 
 		class FGemmCScalar : SHADER_PERMUTATION_ENUM_CLASS("C_SCALAR", EGemmCScalar);
 		class FGemmAlgorithm : SHADER_PERMUTATION_ENUM_CLASS("ALGORITHM", EGemmAlgorithm);
@@ -81,11 +80,6 @@ namespace UE::NNEHlslShaders::Internal
 
 		static void ModifyCompilationEnvironment(const FGlobalShaderPermutationParameters& InParameters, FShaderCompilerEnvironment& OutEnvironment);
 
-		static inline bool ShouldCompilePermutation(const FGlobalShaderPermutationParameters& Parameters)
-		{
-			return FDataDrivenShaderPlatformInfo::GetSupportsNNEShaders(Parameters.Platform);
-		}
-
 		static void FillInParameters(float Alpha, float Beta, int32 TransA, int32 TransB, const NNECore::Internal::FTensor& InputA, const NNECore::Internal::FTensor& InputB, const NNECore::Internal::FTensor* InputC, float CScalar, FParameters& Parameters);
 
 		static void FillInParametersMatMul(const NNECore::Internal::FTensor& InputA, const NNECore::Internal::FTensor& InputB, FParameters& Parameters);
diff --git a/Engine/Plugins/Experimental/NNE/Source/NNEHlslShaders/Internal/NNEHlslShadersInstanceNormalizationCS.h b/Engine/Plugins/Experimental/NNE/Source/NNEHlslShaders/Internal/NNEHlslShadersInstanceNormalizationCS.h
--- a/Engine/Plugins/Experimental/NNE/Source/NNEHlslShaders/Internal/NNEHlslShadersInstanceNormalizationCS.h
+++ b/Engine/Plugins/Experimental/NNE/Source/NNEHlslShaders/Internal/NNEHlslShadersInstanceNormalizationCS.h
@@ -2,10 +2,9 @@
 
 #pragma once
 
-#include "DataDrivenShaderPlatformInfo.h"
-#include "GlobalShader.h"
-#include "ShaderParameterUtils.h"
+#include "NNEHlslShaderBase.h"
 #include "RenderGraphUtils.h"
+#include "ShaderParameterUtils.h"
 
 namespace UE::NNECore::Internal { class FTensor; }
 
@@ -21,10 +20,10 @@ namespace UE::NNEHlslShaders::Internal
 		MAX
 	};
 
-	class NNEHLSLSHADERS_API TInstanceNormalizationCS : public FGlobalShader
+	class NNEHLSLSHADERS_API TInstanceNormalizationCS : public FHlslShaderBase
 	{
 		DECLARE_GLOBAL_SHADER(TInstanceNormalizationCS);
-		SHADER_USE_PARAMETER_STRUCT(TInstanceNormalizationCS, FGlobalShader)
+		SHADER_USE_PARAMETER_STRUCT(TInstanceNormalizationCS, FHlslShaderBase)
 
 		class FInstanceNormalizationAlgorithm : SHADER_PERMUTATION_ENUM_CLASS("ALGORITHM", EInstanceNormalizationAlgorithm);
 		using FPermutationDomain = TShaderPermutationDomain<FInstanceNormalizationAlgorithm>;
@@ -48,10 +47,5 @@ namespace UE::NNEHlslShaders::Internal
 		static EInstanceNormalizationAlgorithm GetAlgorithm(const FParameters& Parameters);
 
 		static void LexFromString(EInstanceNormalizationAlgorithm& OutValue, const TCHAR* StringVal);
-
-		static inline bool ShouldCompilePermutation(const FGlobalShaderPermutationParameters& Parameters)
-		{
-			return FDataDrivenShaderPlatformInfo::GetSupportsNNEShaders(Parameters.Platform);
-		}
 	};
 } // UE::NNEHlslShaders::Internal
\ No newline at end of file
diff --git a/Engine/Plugins/Experimental/NNE/Source/NNEHlslShaders/Internal/NNEHlslShadersNeuralPostProcessingCS.h b/Engine/Plugins/Experimental/NNE/Source/NNEHlslShaders/Internal/NNEHlslShadersNeuralPostProcessingCS.h
--- a/Engine/Plugins/Experimental/NNE/Source/NNEHlslShaders/Internal/NNEHlslShadersNeuralPostProcessingCS.h
+++ b/Engine/Plugins/Experimental/NNE/Source/NNEHlslShaders/Internal/NNEHlslShadersNeuralPostProcessingCS.h
@@ -2,10 +2,9 @@
 
 #pragma once
 
-#include "DataDrivenShaderPlatformInfo.h"
-#include "GlobalShader.h"
-#include "ShaderParameterUtils.h"
+#include "NNEHlslShaderBase.h"
 #include "RenderGraphUtils.h"
+#include "ShaderParameterUtils.h"
 
 namespace UE::NNEHlslShaders::Internal
 {
@@ -29,10 +28,10 @@ namespace UE::NNEHlslShaders::Internal
 		static const int32 THREAD_GROUP_SIZE{ 32 };
 	};
 
-	class NNEHLSLSHADERS_API TNeuralPostProcessingReadInputCS : public FGlobalShader
+	class NNEHLSLSHADERS_API TNeuralPostProcessingReadInputCS : public FHlslShaderBase
 	{
 		DECLARE_GLOBAL_SHADER(TNeuralPostProcessingReadInputCS);
-		SHADER_USE_PARAMETER_STRUCT(TNeuralPostProcessingReadInputCS, FGlobalShader)
+		SHADER_USE_PARAMETER_STRUCT(TNeuralPostProcessingReadInputCS, FHlslShaderBase)
 
 	public:
 
@@ -46,17 +45,12 @@ namespace UE::NNEHlslShaders::Internal
 		END_SHADER_PARAMETER_STRUCT()
 
 		static void ModifyCompilationEnvironment(const FGlobalShaderPermutationParameters& InParameters, FShaderCompilerEnvironment& OutEnvironment);
-
-		static inline bool ShouldCompilePermutation(const FGlobalShaderPermutationParameters& Parameters)
-		{
-			return FDataDrivenShaderPlatformInfo::GetSupportsNNEShaders(Parameters.Platform);
-		}
 	};
 
-	class NNEHLSLSHADERS_API TNeuralPostProcessingPreStepCS : public FGlobalShader
+	class NNEHLSLSHADERS_API TNeuralPostProcessingPreStepCS : public FHlslShaderBase
 	{
 		DECLARE_GLOBAL_SHADER(TNeuralPostProcessingPreStepCS);
-		SHADER_USE_PARAMETER_STRUCT(TNeuralPostProcessingPreStepCS, FGlobalShader)
+		SHADER_USE_PARAMETER_STRUCT(TNeuralPostProcessingPreStepCS, FHlslShaderBase)
 
 	public:
 
@@ -72,17 +66,12 @@ namespace UE::NNEHlslShaders::Internal
 		END_SHADER_PARAMETER_STRUCT()
 
 		static void ModifyCompilationEnvironment(const FGlobalShaderPermutationParameters& InParameters, FShaderCompilerEnvironment& OutEnvironment);
-
-		static inline bool ShouldCompilePermutation(const FGlobalShaderPermutationParameters& Parameters)
-		{
-			return FDataDrivenShaderPlatformInfo::GetSupportsNNEShaders(Parameters.Platform);
-		}
 	};
 
-	class NNEHLSLSHADERS_API TNeuralPostProcessingPostStepCS : public FGlobalShader
+	class NNEHLSLSHADERS_API TNeuralPostProcessingPostStepCS : public FHlslShaderBase
 	{
 		DECLARE_GLOBAL_SHADER(TNeuralPostProcessingPostStepCS);
-		SHADER_USE_PARAMETER_STRUCT(TNeuralPostProcessingPostStepCS, FGlobalShader)
+		SHADER_USE_PARAMETER_STRUCT(TNeuralPostProcessingPostStepCS, FHlslShaderBase)
 
 		class FNeuralPostProcessingOverwrite : SHADER_PERMUTATION_ENUM_CLASS("OVERWRITE", ENeuralPostProcessingOverwrite);
 		class FNeuralPostProcessingInterpolate : SHADER_PERMUTATION_ENUM_CLASS("INTERPOLATE", ENeuralPostProcessingInterpolate);
@@ -102,17 +91,12 @@ namespace UE::NNEHlslShaders::Internal
 		END_SHADER_PARAMETER_STRUCT()
 
 		static void ModifyCompilationEnvironment(const FGlobalShaderPermutationParameters& InParameters, FShaderCompilerEnvironment& OutEnvironment);
-
-		static inline bool ShouldCompilePermutation(const FGlobalShaderPermutationParameters& Parameters)
-		{
-			return FDataDrivenShaderPlatformInfo::GetSupportsNNEShaders(Parameters.Platform);
-		}
 	};
 
-	class NNEHLSLSHADERS_API TNeuralPostProcessingWriteOutputPS : public FGlobalShader
+	class NNEHLSLSHADERS_API TNeuralPostProcessingWriteOutputPS : public FHlslShaderBase
 	{
 		DECLARE_GLOBAL_SHADER(TNeuralPostProcessingWriteOutputPS);
-		SHADER_USE_PARAMETER_STRUCT(TNeuralPostProcessingWriteOutputPS, FGlobalShader)
+		SHADER_USE_PARAMETER_STRUCT(TNeuralPostProcessingWriteOutputPS, FHlslShaderBase)
 
 	public:
 
@@ -124,10 +108,5 @@ namespace UE::NNEHlslShaders::Internal
 		END_SHADER_PARAMETER_STRUCT()
 
 		static void ModifyCompilationEnvironment(const FGlobalShaderPermutationParameters& InParameters, FShaderCompilerEnvironment& OutEnvironment);
-
-		static inline bool ShouldCompilePermutation(const FGlobalShaderPermutationParameters& Parameters)
-		{
-			return FDataDrivenShaderPlatformInfo::GetSupportsNNEShaders(Parameters.Platform);
-		}
 	};
 } // UE::NNEHlslShaders::Internal
\ No newline at end of file
diff --git a/Engine/Plugins/Experimental/NNE/Source/NNEHlslShaders/Internal/NNEHlslShadersPadCS.h b/Engine/Plugins/Experimental/NNE/Source/NNEHlslShaders/Internal/NNEHlslShadersPadCS.h
--- a/Engine/Plugins/Experimental/NNE/Source/NNEHlslShaders/Internal/NNEHlslShadersPadCS.h
+++ b/Engine/Plugins/Experimental/NNE/Source/NNEHlslShaders/Internal/NNEHlslShadersPadCS.h
@@ -2,10 +2,9 @@
 
 #pragma once
 
-#include "DataDrivenShaderPlatformInfo.h"
-#include "GlobalShader.h"
-#include "ShaderParameterUtils.h"
+#include "NNEHlslShaderBase.h"
 #include "RenderGraphUtils.h"
+#include "ShaderParameterUtils.h"
 
 namespace UE::NNEHlslShaders::Internal
 {
@@ -25,10 +24,10 @@ namespace UE::NNEHlslShaders::Internal
 		static const int32 NUM_GROUP_THREADS{ 256 };
 	};
 
-	class NNEHLSLSHADERS_API FPadCS : public FGlobalShader
+	class NNEHLSLSHADERS_API FPadCS : public FHlslShaderBase
 	{
 		DECLARE_GLOBAL_SHADER(FPadCS);
-		SHADER_USE_PARAMETER_STRUCT(FPadCS, FGlobalShader)
+		SHADER_USE_PARAMETER_STRUCT(FPadCS, FHlslShaderBase)
 
 		class FPadMode : SHADER_PERMUTATION_ENUM_CLASS("MODE", EPadMode);
 		class FPadNumDimensions : SHADER_PERMUTATION_RANGE_INT("NUM_DIMENSIONS", 1, FPadConstants::MAX_NUM_DIMENSIONS);
@@ -46,12 +45,6 @@ namespace UE::NNEHlslShaders::Internal
 		END_SHADER_PARAMETER_STRUCT()
 
 		static void ModifyCompilationEnvironment(const FGlobalShaderPermutationParameters& InParameters, FShaderCompilerEnvironment& OutEnvironment);
-
-		static inline bool ShouldCompilePermutation(const FGlobalShaderPermutationParameters& Parameters)
-		{
-			return FDataDrivenShaderPlatformInfo::GetSupportsNNEShaders(Parameters.Platform);
-		}
-
 		static void LexFromString(EPadMode& OutValue, const TCHAR* StringVal);
 	};
 } // UE::NNEHlslShaders::Internal
diff --git a/Engine/Plugins/Experimental/NNE/Source/NNEHlslShaders/Internal/NNEHlslShadersUpsampleCS.h b/Engine/Plugins/Experimental/NNE/Source/NNEHlslShaders/Internal/NNEHlslShadersUpsampleCS.h
--- a/Engine/Plugins/Experimental/NNE/Source/NNEHlslShaders/Internal/NNEHlslShadersUpsampleCS.h
+++ b/Engine/Plugins/Experimental/NNE/Source/NNEHlslShaders/Internal/NNEHlslShadersUpsampleCS.h
@@ -2,10 +2,9 @@
 
 #pragma once
 
-#include "DataDrivenShaderPlatformInfo.h"
-#include "GlobalShader.h"
-#include "ShaderParameterUtils.h"
+#include "NNEHlslShaderBase.h"
 #include "RenderGraphUtils.h"
+#include "ShaderParameterUtils.h"
 
 namespace UE::NNEHlslShaders::Internal
 {
@@ -17,10 +16,10 @@ namespace UE::NNEHlslShaders::Internal
 		static const int32 NUM_GROUP_THREADS{ 256 };
 	};
 
-	class NNEHLSLSHADERS_API FUpsampleCS : public FGlobalShader
+	class NNEHLSLSHADERS_API FUpsampleCS : public FHlslShaderBase
 	{
 		DECLARE_GLOBAL_SHADER(FUpsampleCS);
-		SHADER_USE_PARAMETER_STRUCT(FUpsampleCS, FGlobalShader)
+		SHADER_USE_PARAMETER_STRUCT(FUpsampleCS, FHlslShaderBase)
 
 		class FUpsampleNumDimensions : SHADER_PERMUTATION_RANGE_INT("NUM_DIMENSIONS", 1, FUpsampleConstants::MAX_NUM_DIMENSIONS);
 		using FPermutationDomain = TShaderPermutationDomain<FUpsampleNumDimensions>;
@@ -36,10 +35,5 @@ namespace UE::NNEHlslShaders::Internal
 		END_SHADER_PARAMETER_STRUCT()
 
 		static void ModifyCompilationEnvironment(const FGlobalShaderPermutationParameters& InParameters, FShaderCompilerEnvironment& OutEnvironment);
-
-		static inline bool ShouldCompilePermutation(const FGlobalShaderPermutationParameters& Parameters)
-		{
-			return FDataDrivenShaderPlatformInfo::GetSupportsNNEShaders(Parameters.Platform);
-		}
 	};
 } // UE::NNEHlslShaders::Internal
diff --git a/Engine/Plugins/Experimental/NNE/Source/NNEHlslShaders/Private/NNEHlslShaderBase.cpp b/Engine/Plugins/Experimental/NNE/Source/NNEHlslShaders/Private/NNEHlslShaderBase.cpp
new file mode 100644
--- /dev/null
+++ b/Engine/Plugins/Experimental/NNE/Source/NNEHlslShaders/Private/NNEHlslShaderBase.cpp
@@ -0,0 +1,13 @@
+// Copyright Epic Games, Inc. All Rights Reserved.
+
+#include "NNEHlslShaderBase.h"
+#include "DataDrivenShaderPlatformInfo.h"
+
+namespace UE::NNEHlslShaders::Internal
+{
+	bool FHlslShaderBase::ShouldCompilePermutation(const FGlobalShaderPermutationParameters& Parameters)
+	{
+		return FDataDrivenShaderPlatformInfo::GetSupportsNNEShaders(Parameters.Platform);
+	}
+	
+} // UE::NNEHlslShaders::Internal
\ No newline at end of file
