Fix issues introduced by refactoring the FTensorDescDml

#rnx

diff --git a/Engine/Plugins/Experimental/NNE/Source/NNEQA/Private/NNEQAModel.cpp b/Engine/Plugins/Experimental/NNE/Source/NNEQA/Private/NNEQAModel.cpp
--- a/Engine/Plugins/Experimental/NNE/Source/NNEQA/Private/NNEQAModel.cpp
+++ b/Engine/Plugins/Experimental/NNE/Source/NNEQA/Private/NNEQAModel.cpp
@@ -13,14 +13,21 @@ namespace UE::NNEQA::Private
 {
 namespace FModelQAHelpers
 {
-	FRDGBufferDesc CreateRDGBufferDescForTensorDesc(uint32 ElemByteSize, uint64 SizeInByte)
+	FRDGBufferDesc CreateRDGBufferDescForTensorDesc(uint32 ElemByteSize, uint64 SizeInByte, bool bIsInput)
 	{
 		// FIXME: CreateStructuredDesc() creates a crash on VulkanRHI
 		//FRDGBufferDesc Desc = FRDGBufferDesc::CreateStructuredDesc(ElemByteSize, SizeInByte/ElemByteSize);
 		FRDGBufferDesc Desc = FRDGBufferDesc::CreateBufferDesc(ElemByteSize, SizeInByte/ElemByteSize);
 
 		// FIXME: We should use BUF_SourceCopy for only output buffers (GPU readback)
-		Desc.Usage = EBufferUsageFlags(Desc.Usage | BUF_SourceCopy);
+		if (bIsInput)
+		{
+			Desc.Usage = EBufferUsageFlags(Desc.Usage | BUF_Static);
+		}
+		else
+		{
+			Desc.Usage = EBufferUsageFlags(Desc.Usage | BUF_SourceCopy | BUF_UnorderedAccess);
+		}
 
 		return Desc;
 	}
@@ -39,7 +46,8 @@ namespace FModelQAHelpers
 	void ConvertBinding(FRDGBuilder& GraphBuilder,
 		TConstArrayView<NNECore::FTensorDesc> InTensorDescs,
 		TConstArrayView<NNECore::FTensorBindingCPU> InBindingsCPU,
-		TArray<NNECore::FTensorBindingRDG>& OutBindingsRDG)
+		TArray<NNECore::FTensorBindingRDG>& OutBindingsRDG,
+		bool bIsInput)
 	{
 		check(IsInRenderingThread());
 		check(InTensorDescs.Num() == InBindingsCPU.Num());
@@ -48,7 +56,7 @@ namespace FModelQAHelpers
 		{
 			const NNECore::FTensorDesc& TensorDesc = InTensorDescs[Idx];
 			const NNECore::FTensorBindingCPU& BindingCPU = InBindingsCPU[Idx];
-			const FRDGBufferDesc Desc = CreateRDGBufferDescForTensorDesc(TensorDesc.GetElemByteSize(), BindingCPU.SizeInBytes);
+			const FRDGBufferDesc Desc = CreateRDGBufferDescForTensorDesc(TensorDesc.GetElemByteSize(), BindingCPU.SizeInBytes, bIsInput);
 			FRDGBufferRef TensorBuffer = GraphBuilder.CreateBuffer(Desc, *TensorDesc.GetName(), ERDGBufferFlags::None);
 
 			NNECore::FTensorBindingRDG& BindingRDG = OutBindingsRDG.Emplace_GetRef();
@@ -193,8 +201,8 @@ int FModelQA::EnqueueRDG(FRDGBuilder& RDGBuilder,
 	TArray<NNECore::FTensorBindingRDG> InputBindingsRDG;
 	TArray<NNECore::FTensorBindingRDG> OutputBindingsRDG;
 			
-	FModelQAHelpers::ConvertBinding(RDGBuilder, ModelRDG->GetInputTensorDescs(), InInputBindings, InputBindingsRDG);
-	FModelQAHelpers::ConvertBinding(RDGBuilder, ModelRDG->GetOutputTensorDescs(), InOutputBindings, OutputBindingsRDG);
+	FModelQAHelpers::ConvertBinding(RDGBuilder, ModelRDG->GetInputTensorDescs(), InInputBindings, InputBindingsRDG, true);
+	FModelQAHelpers::ConvertBinding(RDGBuilder, ModelRDG->GetOutputTensorDescs(), InOutputBindings, OutputBindingsRDG, false);
 			
 	FModelQAHelpers::UploadsBindingToGPU(RDGBuilder, InInputBindings, InputBindingsRDG);
 
diff --git a/Engine/Plugins/Experimental/NNE/Source/NNEQA/Resources/NNEQATestDesc.json b/Engine/Plugins/Experimental/NNE/Source/NNEQA/Resources/NNEQATestDesc.json
--- a/Engine/Plugins/Experimental/NNE/Source/NNEQA/Resources/NNEQATestDesc.json
+++ b/Engine/Plugins/Experimental/NNE/Source/NNEQA/Resources/NNEQATestDesc.json
@@ -111,6 +111,7 @@
 					"target": "hs_m05",
 					"runtimes": [
 						{ "name": "NNERuntimeRDGHlsl", "skip": true, "comment": "Model not supported yet." },
+						{ "name": "NNERuntimeRDGDml", "relative_tolerance": 5e-3, "absolute_tolerance": 1e-2 },
 						{ "name": "NNERuntimeORTCuda", "relative_tolerance": 5e-2, "absolute_tolerance": 1e-3 }
 					],
 					"absolute_tolerance": 1e-4
@@ -474,7 +475,7 @@
 		{
 			"category": "BatchNormalization",
 			"targets": [
-				{ "target": "BatchNormalization", "relative_tolerance": 5e-1 }
+				{ "target": "BatchNormalization", "relative_tolerance": 5e-3 }
 			]
 		},
 		{
diff --git a/Engine/Plugins/Experimental/NNE/Source/NNERuntimeRDG/Private/Dml/NNEDmlModel.cpp b/Engine/Plugins/Experimental/NNE/Source/NNERuntimeRDG/Private/Dml/NNEDmlModel.cpp
--- a/Engine/Plugins/Experimental/NNE/Source/NNERuntimeRDG/Private/Dml/NNEDmlModel.cpp
+++ b/Engine/Plugins/Experimental/NNE/Source/NNERuntimeRDG/Private/Dml/NNEDmlModel.cpp
@@ -288,6 +288,8 @@ private:
 	{
 		ID3D12Resource* Resource = DynamicRHI->RHIGetResource(Buffer);
 
+		check(Resource->GetDesc().Width <= Buffer->GetSize());
+
 		return DML_BUFFER_BINDING{ Resource, 0, Buffer->GetSize() };
 	}
 
diff --git a/Engine/Plugins/Experimental/NNE/Source/NNERuntimeRDG/Private/Dml/NNEDmlOperator.cpp b/Engine/Plugins/Experimental/NNE/Source/NNERuntimeRDG/Private/Dml/NNEDmlOperator.cpp
--- a/Engine/Plugins/Experimental/NNE/Source/NNERuntimeRDG/Private/Dml/NNEDmlOperator.cpp
+++ b/Engine/Plugins/Experimental/NNE/Source/NNERuntimeRDG/Private/Dml/NNEDmlOperator.cpp
@@ -113,6 +113,39 @@ FTensorDescDml& FTensorDescDml::SetShape(TConstArrayView<uint32> Shape)
 	return *this;
 }
 
+FTensorDescDml& FTensorDescDml::SetShape(TConstArrayView<uint32> Shape, int32 Rank)
+{
+	check(!bIsValidated);
+	
+	if (!bIsValidated)
+	{
+		Sizes.Reset();
+		
+		if (Shape.Num() > 0)
+		{
+			const int32 Offset = Rank - Shape.Num();
+
+			for (int32 Idx = 0; Idx < Offset; ++Idx)
+			{
+				Sizes.Add(1);
+			}
+
+			Sizes.Append(Shape.GetData(), Shape.Num());
+		}
+		// Handle scalar tensors
+		else if (Shape.Num() == 0)
+		{
+			for (int32 Idx = 0; Idx < Rank; ++Idx)
+			{
+				Sizes.Add(1);
+			}
+		}
+	}
+
+	return *this;
+}
+
+
 FTensorDescDml& FTensorDescDml::SetShape(TConstArrayView<uint32> Shape, TConstArrayView<uint32> BroadcastShape)
 {
 	check(!bIsValidated);
@@ -215,7 +248,7 @@ FTensorDescDml& FTensorDescDml::SetDataType(ENNETensorDataType DataType)
 		DML_TENSOR_DATA_TYPE DmlDataType = Util::GetTensorDataType(DataType);
 		check(DmlDataType != DML_TENSOR_DATA_TYPE_UNKNOWN);
 	
-		uint64 ElemByteSize = CalculateBufferSize(UE::NNECore::GetTensorDataTypeSizeInBytes(DataType));
+		uint64 ElemByteSize = UE::NNECore::GetTensorDataTypeSizeInBytes(DataType);
 		check(ElemByteSize > 0)
 
 		BuffDesc.DataType = DmlDataType;
diff --git a/Engine/Plugins/Experimental/NNE/Source/NNERuntimeRDG/Private/Dml/NNEDmlOperator.h b/Engine/Plugins/Experimental/NNE/Source/NNERuntimeRDG/Private/Dml/NNEDmlOperator.h
--- a/Engine/Plugins/Experimental/NNE/Source/NNERuntimeRDG/Private/Dml/NNEDmlOperator.h
+++ b/Engine/Plugins/Experimental/NNE/Source/NNERuntimeRDG/Private/Dml/NNEDmlOperator.h
@@ -148,19 +148,28 @@ public:
 
 	FTensorDescDml();
 
-	// Set DML operator tensor minimum and maximum rank.
+	// Set tensor minimum and maximum rank.
 	// By default minimum and maximum range is between 1 and DML_TENSOR_DIMENSION_COUNT_MAX1 (8).
 	FTensorDescDml& SetTensorRank(int32 MinRank, int32 MaxRank);
 
 	// Set tensor shape
 	FTensorDescDml& SetShape(TConstArrayView<uint32> Shape);
 
-	// Utility method that calls SetShape() from above
+	// Utility method that calls SetShape(Shape)
 	FTensorDescDml& SetShape(const UE::NNECore::FTensorShape& Shape)
 	{
 		return SetShape(Shape.GetData());
 	}
 
+	// Set shape that and make sure that it has the given rank
+	FTensorDescDml& SetShape(TConstArrayView<uint32> Shape, int32 Rank);
+
+	// Utility method that calls SetShape(Shape, Rank)
+	FTensorDescDml& SetShape(const UE::NNECore::FTensorShape& Shape, int32 Rank)
+	{
+		return SetShape(Shape.GetData(), Rank);
+	}
+
 	// Set shape and match the broadcast shape
 	FTensorDescDml& SetShape(TConstArrayView<uint32> Shape, TConstArrayView<uint32> BroadcastShape);
 
@@ -199,6 +208,14 @@ public:
 			.SetDataOwnedByDml(Tensor.HasPreparedData());
 	}
 
+	// Utility method to use FTensor to set all the members
+	FTensorDescDml& SetFromTensor(const UE::NNECore::Internal::FTensor& Tensor, int32 Rank)
+	{
+		return SetShape(Tensor.GetShape(), Rank)
+			.SetDataType(Tensor.GetDataType())
+			.SetDataOwnedByDml(Tensor.HasPreparedData());
+	}
+
 	// Utility method to use FTensor to set all the members
 	FTensorDescDml& SetFromTensorBroadcast(const UE::NNECore::Internal::FTensor& Tensor, const UE::NNECore::FTensorShape& BroadcastShape)
 	{
@@ -210,7 +227,7 @@ public:
 	// Utility method to use FTensor to set all the members
 	FTensorDescDml& SetFromTensor1D(const UE::NNECore::Internal::FTensor& Tensor, int32 Rank)
 	{
-		return SetShape1D(Tensor.GetShape().Rank(), Rank)
+		return SetShape1D(Tensor.GetShape().GetData()[0], Rank)
 			.SetDataType(Tensor.GetDataType())
 			.SetDataOwnedByDml(Tensor.HasPreparedData());
 	}
diff --git a/Engine/Plugins/Experimental/NNE/Source/NNERuntimeRDG/Private/Dml/NNEDmlOperatorConvolution.cpp b/Engine/Plugins/Experimental/NNE/Source/NNERuntimeRDG/Private/Dml/NNEDmlOperatorConvolution.cpp
--- a/Engine/Plugins/Experimental/NNE/Source/NNERuntimeRDG/Private/Dml/NNEDmlOperatorConvolution.cpp
+++ b/Engine/Plugins/Experimental/NNE/Source/NNERuntimeRDG/Private/Dml/NNEDmlOperatorConvolution.cpp
@@ -246,7 +246,7 @@ public:
 
 		if (!DmlFilterTensorDesc
 				.SetTensorRank(3, 5)
-				.SetFromTensor(InputTensor)
+				.SetFromTensor(FilterTensor)
 				.Validate())
 		{
 			UE_LOG(LogNNE, Error, TEXT("Failed to initialize tensor(s) for DML inference"));
diff --git a/Engine/Plugins/Experimental/NNE/Source/NNERuntimeRDG/Private/Dml/NNEDmlOperatorGather.cpp b/Engine/Plugins/Experimental/NNE/Source/NNERuntimeRDG/Private/Dml/NNEDmlOperatorGather.cpp
--- a/Engine/Plugins/Experimental/NNE/Source/NNERuntimeRDG/Private/Dml/NNEDmlOperatorGather.cpp
+++ b/Engine/Plugins/Experimental/NNE/Source/NNERuntimeRDG/Private/Dml/NNEDmlOperatorGather.cpp
@@ -92,7 +92,7 @@ public:
 		}
 
 		if (!DmlIndicesTensorDesc
-				.SetFromTensor(IndicesTensor)
+				.SetFromTensor(IndicesTensor, InputTensor.GetShape().Rank())
 				.Validate())
 		{
 			UE_LOG(LogNNE, Error, TEXT("Failed to initialize tensor(s) for DML inference"));
diff --git a/Engine/Plugins/Experimental/NNE/Source/NNERuntimeRDG/Private/Dml/NNEDmlOperatorInstanceNormalization.cpp b/Engine/Plugins/Experimental/NNE/Source/NNERuntimeRDG/Private/Dml/NNEDmlOperatorInstanceNormalization.cpp
--- a/Engine/Plugins/Experimental/NNE/Source/NNERuntimeRDG/Private/Dml/NNEDmlOperatorInstanceNormalization.cpp
+++ b/Engine/Plugins/Experimental/NNE/Source/NNERuntimeRDG/Private/Dml/NNEDmlOperatorInstanceNormalization.cpp
@@ -59,7 +59,7 @@ public:
 
 			if (!DmlScalingTensorDesc
 					.SetTensorRank(4, 4)
-					.SetFromTensor(ScaleTensor)
+					.SetFromTensor1D(ScaleTensor, InputTensor.GetShape().Rank())
 					.Validate())
 			{
 				UE_LOG(LogNNE, Error, TEXT("Failed to initialize tensor(s) for DML inference"));
@@ -73,7 +73,7 @@ public:
 
 			if (!DmlBiasTensorDesc
 					.SetTensorRank(4, 4)
-					.SetFromTensor(BiasTensor)
+					.SetFromTensor1D(BiasTensor, InputTensor.GetShape().Rank())
 					.Validate())
 			{
 				UE_LOG(LogNNE, Error, TEXT("Failed to initialize tensor(s) for DML inference"));
diff --git a/Engine/Plugins/Experimental/NNE/Source/NNERuntimeRDG/Private/Dml/NNEDmlOperatorTranspose.cpp b/Engine/Plugins/Experimental/NNE/Source/NNERuntimeRDG/Private/Dml/NNEDmlOperatorTranspose.cpp
--- a/Engine/Plugins/Experimental/NNE/Source/NNERuntimeRDG/Private/Dml/NNEDmlOperatorTranspose.cpp
+++ b/Engine/Plugins/Experimental/NNE/Source/NNERuntimeRDG/Private/Dml/NNEDmlOperatorTranspose.cpp
@@ -70,7 +70,7 @@ public:
 		DmlInputTensorDesc.SetStridesFromShape(InputTensors[0].GetShape());
 		DmlInputTensorDesc.SetStrides(PermuteFunc(DmlInputTensorDesc.GetStrides()));
 
-		if (DmlInputTensorDesc.Validate())
+		if (!DmlInputTensorDesc.Validate())
 		{
 			UE_LOG(LogNNE, Error, TEXT("Failed to initialize Transpose input for DML inference"));
 			return false;
