NNE Reenable parametric tests + stabilize for ORTCuda

Tests: NNE tests (all but RDGDml)

#rnx

diff --git a/Engine/Plugins/Experimental/NNE/Source/NNEQA/Private/NNEQAJsonUtils.h b/Engine/Plugins/Experimental/NNE/Source/NNEQA/Private/NNEQAJsonUtils.h
--- a/Engine/Plugins/Experimental/NNE/Source/NNEQA/Private/NNEQAJsonUtils.h
+++ b/Engine/Plugins/Experimental/NNE/Source/NNEQA/Private/NNEQAJsonUtils.h
@@ -43,9 +43,9 @@ namespace Json
 		FString Name;
 		float AbsoluteTolerance = JSON_TOLERANCE_NOTSET;
 		float RelativeTolerance = JSON_TOLERANCE_NOTSET;
-		bool Skip;
-		bool SkipStatic;
-		bool SkipVariadic;
+		bool Skip = false;
+		bool SkipStatic = false;
+		bool SkipVariadic = false;
 
 		BEGIN_JSON_SERIALIZER
 			JSON_SERIALIZE("name", Name);
@@ -157,7 +157,7 @@ namespace Json
 		TArray<FString> Tags;
 		TArray<FString> AdditionalDatasets;
 		TArray<FString> RemovedDatasets;
-		bool Skip;
+		bool Skip = false;
 		float AbsoluteTolerance = JSON_TOLERANCE_NOTSET;
 		float RelativeTolerance = JSON_TOLERANCE_NOTSET;
 		FString InputType;
@@ -185,8 +185,8 @@ namespace Json
 		TArray<FString> RemovedDatasets;
 		TArray<FTestConfigTarget> Targets;
 		TArray<FTestConfigRuntime> Runtimes;
-		bool Skip;
-		bool IsModelTest;
+		bool Skip = false;
+		bool IsModelTest = false;
 
 		BEGIN_JSON_SERIALIZER
 			JSON_SERIALIZE("category", Category);
diff --git a/Engine/Plugins/Experimental/NNE/Source/NNEQA/Resources/NNEQATestDesc.json b/Engine/Plugins/Experimental/NNE/Source/NNEQA/Resources/NNEQATestDesc.json
--- a/Engine/Plugins/Experimental/NNE/Source/NNEQA/Resources/NNEQATestDesc.json
+++ b/Engine/Plugins/Experimental/NNE/Source/NNEQA/Resources/NNEQATestDesc.json
@@ -8,7 +8,7 @@
 				{
 					"target": "udnie-160-small-720p",
 					"runtimes": [
-						{ "name": "NNERuntimeORTCuda", "relative_tolerance": 1e-3 }
+						{ "name": "NNERuntimeORTCuda", "relative_tolerance": 5e-3 }
 					],
 					"absolute_tolerance": 5.0e-2
 				}
@@ -275,7 +275,10 @@
 				{	"target": "BitShift", "skip": true, "comment": "TODO need attributes" },
 				{	"target": "Ceil" },
 				{	"target": "Celu", 
-					"runtimes": [{ "name": "NNERuntimeRDGHlsl",   "skip": true }],
+					"runtimes": [
+						{ "name": "NNERuntimeRDGHlsl", "skip": true },
+						{ "name": "NNERuntimeORTCuda", "absolute_tolerance":  1e-4}
+					],
 					"additional_datasets": [ "Elementwise.Elu.AttributeTest" ]
 				},
 				{	"target": "Clip", "additional_datasets": [ "Elementwise.Clip.AttributeTest" ] },
@@ -429,6 +432,9 @@
 		},
 		{
 			"category": "Convolution",
+			"runtimes": [
+				{ "name": "NNERuntimeORTCuda", "relative_tolerance": 1e-3, "absolute_tolerance": 1e-4 }
+			],
 			"targets": [
 				{ "target": "Conv", "absolute_tolerance": 1e-4, "relative_tolerance": 1e-4 }
 			]
@@ -456,6 +462,9 @@
 		},
 		{
 			"category": "Identity",
+			"runtimes": [
+				{ "name": "NNERuntimeORTDml", "skip": true, "comment": "Error: Failed to find kernel args" }
+			],
 			"targets": [
 				{ "target": "Identity" }
 			]
@@ -706,6 +715,7 @@
 			"datasets": [
 				{ "inputs": [ { "shape": [ 1,0,3 ] } ],
 					"runtimes": [
+						{ "name": "NNERuntimeORTCuda", "skip": true, "comment": "ORTCuda does not support 0 sized tensor for Shrink" },
 						{ "name": "NNERuntimeRDGHlsl", "skip": true, "comment": "RDG does not support 0 sized tensor" },
 						{ "name": "NNERuntimeRDGDml",  "skip": true, "comment": "RDG does not support 0 sized tensor" }
 					]}
