Fixed typos in RDG DML Concat and MatMul operators

#rnx

diff --git a/Engine/Plugins/Experimental/NNE/Source/NNERuntimeRDG/Private/Dml/NNEDmlOperatorConcat.cpp b/Engine/Plugins/Experimental/NNE/Source/NNERuntimeRDG/Private/Dml/NNEDmlOperatorConcat.cpp
--- a/Engine/Plugins/Experimental/NNE/Source/NNERuntimeRDG/Private/Dml/NNEDmlOperatorConcat.cpp
+++ b/Engine/Plugins/Experimental/NNE/Source/NNERuntimeRDG/Private/Dml/NNEDmlOperatorConcat.cpp
@@ -64,6 +64,8 @@ public:
 					UE_LOG(LogNNE, Error, TEXT("Failed to initialize Concat input for DML inference"));
 					return false;
 				}
+
+				DmlInputDescs[Idx] = *InputDescs[Idx].GetDmlDesc();
 			}
 		}
 		
diff --git a/Engine/Plugins/Experimental/NNE/Source/NNERuntimeRDG/Private/Dml/NNEDmlOperatorMatMul.cpp b/Engine/Plugins/Experimental/NNE/Source/NNERuntimeRDG/Private/Dml/NNEDmlOperatorMatMul.cpp
--- a/Engine/Plugins/Experimental/NNE/Source/NNERuntimeRDG/Private/Dml/NNEDmlOperatorMatMul.cpp
+++ b/Engine/Plugins/Experimental/NNE/Source/NNERuntimeRDG/Private/Dml/NNEDmlOperatorMatMul.cpp
@@ -221,6 +221,7 @@ public:
 
 		if (!DmlOutputTensorDesc
 				.SetTensorRank(2, 4)
+				.SetFromTensor(OutTensor)
 				.SetShape(DmlShapeOut)
 				.Validate())
 		{
