NNE ModelData import settings + remove serialisation debug log

What:
* Add target runtimes setting and related editor (multi edition via tabs)
* Add target runtimes and source path as asset tags (visible as asset tooltips/column)
* Add import/reimport capability
* Add open source path/with external editor capability
* Remove NNEModel deserialisation log

What it is NOT:
* The runtime settings are NOT yet used by the cooking process, meaning no reduction in cooked asset size yet.

How:
* Add asset action, editor and runtime widge for NNEModelData
* Add editor only metadata to NNEModelData (target runtimes and AssetImportData)
* Bumped asset version with backward compatibility, all runtimes  + no reimport / open source capability for old assets.
* Editing multiple asset will spawn one editor tabs for each.
* One can remove a not registered runtime from an asset but not add them.

Tests:
* Asset backward compat
* Import/reimport (keeping metadata) + multi import
* Saving/Autosaving of assets
* Add remove runtimes in editor, verify tags and serialisation
* Visualize and remove not currently registered runtimes from assets
* Delete external source file + reimport / open source
* Style transfer demo with reimported assets
* NNE tests (green but the known one being investigated separately)

#rnx

diff --git a/Engine/Plugins/Experimental/NNE/Source/NNECore/Private/NNECoreModelData.cpp b/Engine/Plugins/Experimental/NNE/Source/NNECore/Private/NNECoreModelData.cpp
--- a/Engine/Plugins/Experimental/NNE/Source/NNECore/Private/NNECoreModelData.cpp
+++ b/Engine/Plugins/Experimental/NNE/Source/NNECore/Private/NNECoreModelData.cpp
@@ -8,6 +8,7 @@
 #include "NNECoreRuntimeFormat.h"
 #include "Serialization/CustomVersion.h"
 #include "UObject/WeakInterfacePtr.h"
+#include "EditorFramework/AssetImportData.h"
 
 #if WITH_EDITOR
 #include "Containers/StringFwd.h"
@@ -19,8 +20,18 @@
 #include "Misc/Guid.h"
 #endif
 
+enum Type
+{
+	Initial = 0,
+	TargetRuntimesAndAssetImportData = 1,
+	// -----<new versions can be added before this line>-------------------------------------------------
+	// - this needs to be the last line (see note below)
+	VersionPlusOne,
+	LatestVersion = VersionPlusOne - 1
+};
+
 const FGuid UNNEModelData::GUID(0x9513202e, 0xeba1b279, 0xf17fe5ba, 0xab90c3f2);
-FCustomVersionRegistration NNEModelDataVersion(UNNEModelData::GUID, 0, TEXT("NNEModelDataVersion"));
+FCustomVersionRegistration NNEModelDataVersion(UNNEModelData::GUID, LatestVersion, TEXT("NNEModelDataVersion"));//Always save with the latest version
 
 #if WITH_EDITOR
 
@@ -131,9 +142,6 @@ void UNNEModelData::Serialize(FArchive& Ar)
 	// Store the asset version (no effect in load)
 	Ar.UsingCustomVersion(UNNEModelData::GUID);
 
-	// Only one version is supported for now
-	check(Ar.CustomVer(UNNEModelData::GUID) == 0);
-
 	// Recreate each model data when cooking
 	if (Ar.IsCooking() && Ar.IsSaving())
 	{
@@ -169,13 +177,24 @@ void UNNEModelData::Serialize(FArchive& Ar)
 			Ar << RuntimeNames[i];
 			Ar << ModelData[RuntimeNames[i]];
 		}
-
-		UE_LOG(LogTemp, Display, TEXT("UNNEModelData: Serialized data of %d runtimes"), NumItems);
 	}
 	else
 	{
 		int32 NumItems = 0;
 
+#if WITH_EDITORONLY_DATA
+		if (Ar.CustomVer(UNNEModelData::GUID) >= TargetRuntimesAndAssetImportData)
+		{
+			Ar << TargetRuntimes;
+			Ar << AssetImportData;
+		}
+		else 
+		{
+			// AssetImportData should always be valid
+			AssetImportData = NewObject<UAssetImportData>(this, TEXT("AssetImportData"));
+		}
+#endif //WITH_EDITORONLY_DATA
+
 		Ar << FileType;
 		Ar << FileData;
 		Ar << FileDataId;
@@ -191,8 +210,71 @@ void UNNEModelData::Serialize(FArchive& Ar)
 				Ar << Data;
 				ModelData.Add(Name, MoveTemp(Data));
 			}
+		}
+	}
+}
+
+#if WITH_EDITORONLY_DATA
+namespace UE::NNECore::ModelDataHelpers
+{
+	FString GetRuntimesAsString(TArrayView<const FString> Runtimes)
+	{
+		if (Runtimes.Num() == 0)
+		{
+			return TEXT("All");
+		}
+
+		FString RuntimesAsOneString;
+		bool bIsFirstRuntime = true;
+
+		for (const FString& Runtime : Runtimes)
+		{
+			if (!bIsFirstRuntime)
+			{
+				RuntimesAsOneString += TEXT(", ");
+			}
+			RuntimesAsOneString += Runtime;
+			bIsFirstRuntime = false;
+		}
+		return RuntimesAsOneString;
+	}
+} // UE::NNECore::ModelDataHelpers
 
-			UE_LOG(LogTemp, Display, TEXT("UNNEModelData: Deserialized data of %d runtimes"), NumItems);
+void UNNEModelData::PostInitProperties()
+{
+	if (!HasAnyFlags(RF_ClassDefaultObject))
+	{
+		AssetImportData = NewObject<UAssetImportData>(this, TEXT("AssetImportData"));
+	}
+	Super::PostInitProperties();
+}
+
+void UNNEModelData::GetAssetRegistryTags(TArray<FAssetRegistryTag>& OutTags) const
+{
+	if (AssetImportData)
+	{
+		OutTags.Add(FAssetRegistryTag(SourceFileTagName(), AssetImportData->GetSourceData().ToJson(), FAssetRegistryTag::TT_Hidden));
+	}
+
+	OutTags.Add(FAssetRegistryTag("TargetRuntimes", UE::NNECore::ModelDataHelpers::GetRuntimesAsString(GetTargetRuntimes()), FAssetRegistryTag::TT_Alphabetical));
+
+	Super::GetAssetRegistryTags(OutTags);
+}
+
+void UNNEModelData::SetTargetRuntimes(TArrayView<const FString> RuntimeNames)
+{
+	TargetRuntimes = RuntimeNames;
+
+	TArray<FString, TInlineAllocator<10>> CookedRuntimes;
+	ModelData.GetKeys(CookedRuntimes);
+	for (const FString& Runtime : CookedRuntimes)
+	{
+		if (!TargetRuntimes.Contains(Runtime))
+		{
+			ModelData.Remove(Runtime);
 		}
 	}
-}
\ No newline at end of file
+	ModelData.Compact();
+}
+
+#endif //WITH_EDITORONLY_DATA
\ No newline at end of file
diff --git a/Engine/Plugins/Experimental/NNE/Source/NNECore/Public/NNECoreModelData.h b/Engine/Plugins/Experimental/NNE/Source/NNECore/Public/NNECoreModelData.h
--- a/Engine/Plugins/Experimental/NNE/Source/NNECore/Public/NNECoreModelData.h
+++ b/Engine/Plugins/Experimental/NNE/Source/NNECore/Public/NNECoreModelData.h
@@ -40,8 +40,6 @@ public:
 	 */
 	TConstArrayView<uint8> GetModelData(const FString& RuntimeName);
 
-public:
-
 	/**
 	 * Implements custom serialization of this asset.
 	 * @param Ar The archive to serialize from/to.
@@ -53,6 +51,39 @@ public:
 	 */
 	const static FGuid GUID;
 
+#if WITH_EDITORONLY_DATA
+	// UObject interface
+	virtual void PostInitProperties() override;
+	virtual void GetAssetRegistryTags(TArray<FAssetRegistryTag>& OutTags) const override;
+	// End of UObject interface
+
+	/**
+	 * Get the target runtimes this model data will be cooked for. An empty list mean all runtimes.
+	 *
+	 * @return The target runtimes names.
+	 */
+	TArrayView<const FString> GetTargetRuntimes() const { return TargetRuntimes; }
+
+	/**
+	 * Set the target runtimes this model data will be cooked for. An empty list mean all runtimes.
+	 *
+	 * @param RuntimeNames The target runtimes names.
+	 */
+	void SetTargetRuntimes(TArrayView<const FString> RuntimeNames);
+
+	/**
+	 * Importing data used for this asset.
+	 */
+	UPROPERTY(VisibleAnywhere, Instanced, Category = ImportSettings)
+	TObjectPtr<class UAssetImportData> AssetImportData;
+
+private:
+	/**
+	 * A list of string of the supported runtime, empty to support them all.
+	 */
+	TArray<FString> TargetRuntimes;
+#endif // WITH_EDITORONLY_DATA
+
 private:
 
 	/**
diff --git a/Engine/Plugins/Experimental/NNE/Source/NNEEditorTools/NNEEditorTools.Build.cs b/Engine/Plugins/Experimental/NNE/Source/NNEEditorTools/NNEEditorTools.Build.cs
--- a/Engine/Plugins/Experimental/NNE/Source/NNEEditorTools/NNEEditorTools.Build.cs
+++ b/Engine/Plugins/Experimental/NNE/Source/NNEEditorTools/NNEEditorTools.Build.cs
@@ -23,7 +23,10 @@ public class NNEEditorTools : ModuleRules
 			PrivateDependencyModuleNames.AddRange(
 				new string[]
 				{
-					"UnrealEd"
+					"Slate",
+					"SlateCore",
+					"UnrealEd",
+					"AssetTools"
 				}
 			);
 		}
diff --git a/Engine/Plugins/Experimental/NNE/Source/NNEEditorTools/Private/NNEEditorTools.cpp b/Engine/Plugins/Experimental/NNE/Source/NNEEditorTools/Private/NNEEditorTools.cpp
--- a/Engine/Plugins/Experimental/NNE/Source/NNEEditorTools/Private/NNEEditorTools.cpp
+++ b/Engine/Plugins/Experimental/NNE/Source/NNEEditorTools/Private/NNEEditorTools.cpp
@@ -1,6 +1,7 @@
 // Copyright Epic Games, Inc. All Rights Reserved.
 
 #include "Modules/ModuleManager.h"
+#include "NNEEditorToolsModelDataActions.h"
 
 class FNNEEditorToolsModule : public IModuleInterface
 {
@@ -8,12 +9,19 @@ public:
 
 	virtual void StartupModule() override
 	{
+		ModelDataAssetTypeActions = MakeShared<UE::NNEEditorTools::Private::FModelDataAssetTypeActions>();
+		FAssetToolsModule::GetModule().Get().RegisterAssetTypeActions(ModelDataAssetTypeActions.ToSharedRef());
 	}
 
 	virtual void ShutdownModule() override
 	{
+		if (FModuleManager::Get().IsModuleLoaded("AssetTools"))
+		{
+			FAssetToolsModule::GetModule().Get().UnregisterAssetTypeActions(ModelDataAssetTypeActions.ToSharedRef());
+		}
 	}
 
+	TSharedPtr<UE::NNEEditorTools::Private::FModelDataAssetTypeActions> ModelDataAssetTypeActions;
 };
 
 IMPLEMENT_MODULE(FNNEEditorToolsModule, NNEEditorTools)
\ No newline at end of file
diff --git a/Engine/Plugins/Experimental/NNE/Source/NNEEditorTools/Private/NNEEditorToolsModelDataActions.cpp b/Engine/Plugins/Experimental/NNE/Source/NNEEditorTools/Private/NNEEditorToolsModelDataActions.cpp
new file mode 100644
--- /dev/null
+++ b/Engine/Plugins/Experimental/NNE/Source/NNEEditorTools/Private/NNEEditorToolsModelDataActions.cpp
@@ -0,0 +1,59 @@
+// Copyright Epic Games, Inc. All Rights Reserved.
+
+#include "NNEEditorToolsModelDataActions.h"
+#include "EditorFramework/AssetImportData.h"
+#include "NNECoreModelData.h"
+#include "NNEEditorToolsModelDataEditorToolkit.h"
+
+#define LOCTEXT_NAMESPACE "AssetTypeActions"
+
+namespace UE::NNEEditorTools::Private
+{
+
+	UClass* FModelDataAssetTypeActions::GetSupportedClass() const
+	{
+		return UNNEModelData::StaticClass();
+	}
+
+	FText FModelDataAssetTypeActions::GetName() const
+	{
+		return LOCTEXT("FModelDataAssetTypeActionsName", "NNE Model Data");
+	}
+
+	FColor FModelDataAssetTypeActions::GetTypeColor() const
+	{
+		return FColor::Cyan;
+	}
+
+	uint32 FModelDataAssetTypeActions::GetCategories()
+	{
+		return EAssetTypeCategories::Misc;
+	}
+
+	void FModelDataAssetTypeActions::OpenAssetEditor(const TArray<UObject*>& InObjects, TSharedPtr<class IToolkitHost> EditWithinLevelEditor)
+	{
+		for (UObject* Obj : InObjects)
+		{
+			UNNEModelData* ModelData = Cast<UNNEModelData>(Obj);
+			if (ModelData != nullptr)
+			{
+				MakeShared<FModelDataEditorToolkit>()->InitEditor(ModelData);
+			}
+		}
+	}
+
+	void FModelDataAssetTypeActions::GetResolvedSourceFilePaths(const TArray<UObject*>& TypeAssets, TArray<FString>& OutSourceFilePaths) const
+	{
+		for (UObject* Asset : TypeAssets)
+		{
+			const UNNEModelData* ModelData = CastChecked<UNNEModelData>(Asset);
+			if (ModelData->AssetImportData)
+			{
+				ModelData->AssetImportData->ExtractFilenames(OutSourceFilePaths);
+			}
+		}
+	}
+
+} // UE::NNEEditorTools::Private
+
+#undef LOCTEXT_NAMESPACE
\ No newline at end of file
diff --git a/Engine/Plugins/Experimental/NNE/Source/NNEEditorTools/Private/NNEEditorToolsModelDataActions.h b/Engine/Plugins/Experimental/NNE/Source/NNEEditorTools/Private/NNEEditorToolsModelDataActions.h
new file mode 100644
--- /dev/null
+++ b/Engine/Plugins/Experimental/NNE/Source/NNEEditorTools/Private/NNEEditorToolsModelDataActions.h
@@ -0,0 +1,26 @@
+// Copyright Epic Games, Inc. All Rights Reserved.
+
+#pragma once
+
+#include "CoreMinimal.h"
+#include "AssetTypeActions_Base.h"
+
+class UNNEModelData;
+
+namespace UE::NNEEditorTools::Private
+{
+	class FModelDataAssetTypeActions : public FAssetTypeActions_Base
+	{
+	public:
+		// IAssetTypeActions interface
+		virtual FText GetName() const override;
+		virtual FColor GetTypeColor() const override;
+		virtual UClass* GetSupportedClass() const override;
+		virtual uint32 GetCategories() override;
+		virtual bool IsImportedAsset() const override { return true; }
+		virtual void GetResolvedSourceFilePaths(const TArray<UObject*>& TypeAssets, TArray<FString>& OutSourceFilePaths) const override;
+		virtual void OpenAssetEditor(const TArray<UObject*>& InObjects, TSharedPtr<class IToolkitHost> EditWithinLevelEditor) override;
+		// End of IAssetTypeActions interface
+	};
+
+} // UE::NNEEditorTools::Private
\ No newline at end of file
diff --git a/Engine/Plugins/Experimental/NNE/Source/NNEEditorTools/Private/NNEEditorToolsModelDataEditorToolkit.cpp b/Engine/Plugins/Experimental/NNE/Source/NNEEditorTools/Private/NNEEditorToolsModelDataEditorToolkit.cpp
new file mode 100644
--- /dev/null
+++ b/Engine/Plugins/Experimental/NNE/Source/NNEEditorTools/Private/NNEEditorToolsModelDataEditorToolkit.cpp
@@ -0,0 +1,69 @@
+// Copyright Epic Games, Inc. All Rights Reserved.
+
+#include "NNEEditorToolsModelDataEditorToolkit.h"
+#include "NNECore.h"
+#include "NNECoreModelData.h"
+#include "SNNEEditorToolsRuntimesWidget.h"
+#include "Widgets/Docking/SDockTab.h"
+
+namespace UE::NNEEditorTools::Private
+{
+
+	void FModelDataEditorToolkit::InitEditor(UNNEModelData* InModelData)
+	{
+		ModelData = InModelData;
+
+		const TSharedRef<FTabManager::FLayout> Layout = FTabManager::NewLayout("NNEModelDataEditorLayout")
+			->AddArea
+			(
+				FTabManager::NewPrimaryArea()->SetOrientation(Orient_Vertical)
+				->Split
+				(
+					FTabManager::NewStack()
+					->SetSizeCoefficient(0.8f)
+					->AddTab(TEXT("NNEModelDataTab"), ETabState::OpenedTab)
+				)
+			);
+
+		FAssetEditorToolkit::InitAssetEditor(EToolkitMode::Standalone, {}, TEXT("NNEModelDataEditor"), Layout, true, true, ModelData);
+	}
+
+	void FModelDataEditorToolkit::RegisterTabSpawners(const TSharedRef<class FTabManager>& InTabManager)
+	{
+		FAssetEditorToolkit::RegisterTabSpawners(InTabManager);
+
+		WorkspaceMenuCategory = InTabManager->AddLocalWorkspaceMenuCategory(INVTEXT("NNE Model Data Editor"));
+
+		InTabManager->RegisterTabSpawner("NNEModelDataTab", FOnSpawnTab::CreateLambda([=](const FSpawnTabArgs&)
+			{
+				return SNew(SDockTab)
+				[
+					SNew(SRuntimesWidget)
+					.TargetRuntimes(this, &FModelDataEditorToolkit::GetTargetRuntimes)
+					.OnTargetRuntimesChanged(this, &FModelDataEditorToolkit::SetTargetRuntimes)
+				];
+			}))
+			.SetDisplayName(INVTEXT("NNE Model Data"))
+			.SetGroup(WorkspaceMenuCategory.ToSharedRef());
+	}
+
+	void FModelDataEditorToolkit::UnregisterTabSpawners(const TSharedRef<class FTabManager>& InTabManager)
+	{
+		FAssetEditorToolkit::UnregisterTabSpawners(InTabManager);
+		InTabManager->UnregisterTabSpawner("NNEModelDataTab");
+	}
+
+	TArrayView<const FString> FModelDataEditorToolkit::GetTargetRuntimes() const
+	{
+		check(ModelData != nullptr);
+		return ModelData->GetTargetRuntimes();
+	}
+
+	void FModelDataEditorToolkit::SetTargetRuntimes(TArrayView<const FString> TargetRuntimes)
+	{
+		check(ModelData != nullptr);
+		ModelData->Modify();
+		ModelData->SetTargetRuntimes(TargetRuntimes);
+	}
+
+} // UE::NNEEditorTools::Private
\ No newline at end of file
diff --git a/Engine/Plugins/Experimental/NNE/Source/NNEEditorTools/Private/NNEEditorToolsModelDataEditorToolkit.h b/Engine/Plugins/Experimental/NNE/Source/NNEEditorTools/Private/NNEEditorToolsModelDataEditorToolkit.h
new file mode 100644
--- /dev/null
+++ b/Engine/Plugins/Experimental/NNE/Source/NNEEditorTools/Private/NNEEditorToolsModelDataEditorToolkit.h
@@ -0,0 +1,32 @@
+// Copyright Epic Games, Inc. All Rights Reserved.
+
+#pragma once
+
+#include "CoreMinimal.h"
+#include "Toolkits/AssetEditorToolkit.h"
+
+class UNNEModelData;
+
+namespace UE::NNEEditorTools::Private
+{
+	class FModelDataEditorToolkit : public FAssetEditorToolkit
+	{
+	public:
+		void InitEditor(UNNEModelData* ModelData);
+
+		virtual void RegisterTabSpawners(const TSharedRef<class FTabManager>& TabManager) override;
+		virtual void UnregisterTabSpawners(const TSharedRef<class FTabManager>& TabManager) override;
+		virtual FName GetToolkitFName() const override { return FName("NNEModelDataEditor"); }
+		virtual FText GetBaseToolkitName() const override { return INVTEXT("NNE Model Data Editor"); }
+		virtual FString GetWorldCentricTabPrefix() const override { return TEXT("NNE Model Data Editor "); }
+		virtual FLinearColor GetWorldCentricTabColorScale() const override { return {}; }
+
+	private:
+		TArrayView<const FString> GetTargetRuntimes() const;
+		void SetTargetRuntimes(TArrayView<const FString> TargetRuntimes);
+
+		TArray<FString> RegisteredRuntimes;
+		UNNEModelData* ModelData;
+	};
+
+} // UE::NNEEditorTools::Private
\ No newline at end of file
diff --git a/Engine/Plugins/Experimental/NNE/Source/NNEEditorTools/Private/NNEModelDataFactory.cpp b/Engine/Plugins/Experimental/NNE/Source/NNEEditorTools/Private/NNEModelDataFactory.cpp
--- a/Engine/Plugins/Experimental/NNE/Source/NNEEditorTools/Private/NNEModelDataFactory.cpp
+++ b/Engine/Plugins/Experimental/NNE/Source/NNEEditorTools/Private/NNEModelDataFactory.cpp
@@ -1,7 +1,14 @@
 // Copyright Epic Games, Inc. All Rights Reserved.
 
 #include "NNEModelDataFactory.h"
+#include "Editor.h"
+#include "Editor/EditorEngine.h"
+#include "EditorFramework/AssetImportData.h"
+#include "HAL/FileManager.h"
+#include "NNECore.h"
 #include "NNECoreModelData.h"
+#include "Subsystems/ImportSubsystem.h"
+
 
 UNNEModelDataFactory::UNNEModelDataFactory(const FObjectInitializer& ObjectInitializer) : UFactory(ObjectInitializer)
 {
@@ -14,19 +21,104 @@ UNNEModelDataFactory::UNNEModelDataFactory(const FObjectInitializer& ObjectIniti
 
 UObject* UNNEModelDataFactory::FactoryCreateBinary(UClass* Class, UObject* InParent, FName Name, EObjectFlags Flags, UObject* Context, const TCHAR * Type, const uint8 *& Buffer, const uint8 * BufferEnd, FFeedbackContext* Warn)
 {
+	GEditor->GetEditorSubsystem<UImportSubsystem>()->BroadcastAssetPreImport(this, Class, InParent, Name, Type);
+
 	if (!Type || !Buffer || !BufferEnd || BufferEnd - Buffer <= 0)
 	{
+		GEditor->GetEditorSubsystem<UImportSubsystem>()->BroadcastAssetPostImport(this, nullptr);
 		return nullptr;
 	}
 
-	UNNEModelData* Result = NewObject<UNNEModelData>(InParent, Class, Name, Flags);
+	UNNEModelData* ModelData = NewObject<UNNEModelData>(InParent, Class, Name, Flags);
 	TConstArrayView<uint8> BufferView = MakeArrayView(Buffer, BufferEnd-Buffer);
-	Result->Init(Type, BufferView);
+	ModelData->Init(Type, BufferView);
+
+	ModelData->AssetImportData->Update(GetCurrentFilename());
+	GEditor->GetEditorSubsystem<UImportSubsystem>()->BroadcastAssetPostImport(this, ModelData);
 
-	return Result;
+	return ModelData;
 }
 
 bool UNNEModelDataFactory::FactoryCanImport(const FString & Filename)
 {
 	return Filename.EndsWith(FString("onnx"));
+}
+
+bool UNNEModelDataFactory::CanReimport(UObject* Obj, TArray<FString>& OutFilenames)
+{
+	UNNEModelData* pModel = Cast<UNNEModelData>(Obj);
+	if (pModel)
+	{
+		pModel->AssetImportData->ExtractFilenames(OutFilenames);
+		return true;
+	}
+	return false;
+}
+
+void UNNEModelDataFactory::SetReimportPaths(UObject* Obj, const TArray<FString>& NewReimportPaths)
+{
+	UNNEModelData* pModel = Cast<UNNEModelData>(Obj);
+	if (pModel && ensure(NewReimportPaths.Num() == 1))
+	{
+		pModel->AssetImportData->UpdateFilenameOnly(NewReimportPaths[0]);
+	}
+}
+
+EReimportResult::Type UNNEModelDataFactory::Reimport(UObject* Obj)
+{
+	if (!Obj || !Obj->IsA(UNNEModelData::StaticClass()))
+	{
+		return EReimportResult::Failed;
+	}
+
+	UNNEModelData* pModel = Cast<UNNEModelData>(Obj);
+	const FString ResolvedSourceFilePath = pModel->AssetImportData->GetFirstFilename();
+	if (!ResolvedSourceFilePath.Len())
+	{
+		UE_LOG(LogNNE, Warning, TEXT("Cannot reimport NNE Model Data: it does not have a path stored."));
+		return EReimportResult::Failed;
+	}
+
+	// Ensure that the file provided by the path exists
+	if (IFileManager::Get().FileSize(*ResolvedSourceFilePath) == INDEX_NONE)
+	{
+		UE_LOG(LogNNE, Warning, TEXT("Cannot reimport NNE Model Data: source file cannot be found."));
+		return EReimportResult::Failed;
+	}
+
+	// Keep runtime list to be reapplied upon asset import.
+	TArray<FString, TInlineAllocator<10>> TargetRuntimes;
+	TargetRuntimes = pModel->GetTargetRuntimes();
+
+	bool OutCanceled = false;
+
+	if (ImportObject(pModel->GetClass(), pModel->GetOuter(), *pModel->GetName(), RF_Public | RF_Standalone, ResolvedSourceFilePath, nullptr, OutCanceled) != nullptr)
+	{
+		UE_LOG(LogNNE, Log, TEXT("Reimported successfully"));
+
+		pModel->SetTargetRuntimes(TargetRuntimes);
+		pModel->AssetImportData->Update(ResolvedSourceFilePath);
+
+		// Try to find the outer package so we can dirty it up
+		if (pModel->GetOuter())
+		{
+			pModel->GetOuter()->MarkPackageDirty();
+		}
+		else
+		{
+			pModel->MarkPackageDirty();
+		}
+	}
+	else if (OutCanceled)
+	{
+		UE_LOG(LogNNE, Warning, TEXT("Reimport canceled"));
+		return EReimportResult::Cancelled;
+	}
+	else
+	{
+		UE_LOG(LogNNE, Warning, TEXT("Reimport failed"));
+		return EReimportResult::Failed;
+	}
+
+	return EReimportResult::Succeeded;
 }
\ No newline at end of file
diff --git a/Engine/Plugins/Experimental/NNE/Source/NNEEditorTools/Private/SNNEEditorToolsRuntimesWidget.cpp b/Engine/Plugins/Experimental/NNE/Source/NNEEditorTools/Private/SNNEEditorToolsRuntimesWidget.cpp
new file mode 100644
--- /dev/null
+++ b/Engine/Plugins/Experimental/NNE/Source/NNEEditorTools/Private/SNNEEditorToolsRuntimesWidget.cpp
@@ -0,0 +1,145 @@
+// Copyright Epic Games, Inc. All Rights Reserved.
+
+#include "SNNEEditorToolsRuntimesWidget.h"
+#include "NNECore.h"
+#include "SlateOptMacros.h"
+#include "Styling/AppStyle.h"
+#include "Widgets/Input/SCheckBox.h"
+#include "Widgets/Layout/SBox.h"
+#include "Widgets/Layout/SWrapBox.h"
+#include "Widgets/SBoxPanel.h"
+#include "Widgets/Text/STextBlock.h"
+
+
+#define LOCTEXT_NAMESPACE "SNNEEditorToolsRuntimesWidget"
+
+namespace UE::NNEEditorTools::Private
+{
+
+BEGIN_SLATE_FUNCTION_BUILD_OPTIMIZATION
+
+void SRuntimesWidget::Construct(const FArguments& InArgs)
+{
+	TargetRuntimes = InArgs._TargetRuntimes;
+	OnTargetRuntimesChanged = InArgs._OnTargetRuntimesChanged;
+
+	TSharedPtr<SWrapBox> WrapBox;
+
+	ChildSlot
+	[
+		SNew(SBox)
+		[
+			SNew(SVerticalBox)
+			+
+			SVerticalBox::Slot()
+			.AutoHeight()
+			.HAlign(HAlign_Fill)
+			.VAlign(VAlign_Fill)
+			.Padding(4.0)
+			[
+				SNew(STextBlock).Text(FText::FromString("Target Runtimes"))
+				.Font(FAppStyle::GetFontStyle("PropertyWindow.BoldFont"))
+			]
+			+
+			SVerticalBox::Slot()
+			.AutoHeight()
+			.HAlign(HAlign_Fill)
+			.VAlign(VAlign_Fill)
+			.Padding(4.0)
+			[
+				SAssignNew(WrapBox, SWrapBox)
+				.PreferredWidth(300.0f)
+			]
+		]
+	];
+
+	InitCheckBoxNames();
+
+	for (int32 i = 0; i < CheckBoxNames.Num(); ++i)
+	{
+		WrapBox->AddSlot()
+		.Padding(5)
+		.VAlign(VAlign_Top)
+		[
+			SNew(SCheckBox)
+			.Style(FAppStyle::Get(), i==0?"RadioButton":"Checkbox")//"All" in first position as a RadioButton.
+			.IsChecked(this, &SRuntimesWidget::IsEntryChecked, i)
+			.OnCheckStateChanged(this, &SRuntimesWidget::OnEntryChanged, i)
+			[
+				SNew(STextBlock).Text(FText::FromString(CheckBoxNames[i]))
+				.Font(FAppStyle::GetFontStyle("PropertyWindow.NormalFont"))
+			]
+		];
+	}
+}
+
+END_SLATE_FUNCTION_BUILD_OPTIMIZATION
+
+void SRuntimesWidget::OnEntryChanged(ECheckBoxState CheckType, int32 Index)
+{
+	check(Index < CheckBoxNames.Num());
+	const FString& RuntimeName = CheckBoxNames[Index];
+
+	if (!OnTargetRuntimesChanged.IsBound())
+		return;
+
+	if (RuntimeName == TEXT("All"))
+	{
+		if (CheckType == ECheckBoxState::Checked)
+		{
+			TArray<FString> EmptyTargetRuntimes;
+			OnTargetRuntimesChanged.Execute(EmptyTargetRuntimes);
+		}
+	}
+	else
+	{
+		TArray<FString, TInlineAllocator<10>> NewTargetRuntimes;
+		NewTargetRuntimes = TargetRuntimes.Get();
+		if (CheckType == ECheckBoxState::Checked)
+		{
+			NewTargetRuntimes.AddUnique(RuntimeName);
+		}
+		else
+		{
+			NewTargetRuntimes.Remove(RuntimeName);
+		}
+		OnTargetRuntimesChanged.Execute(NewTargetRuntimes);
+	}
+}
+
+ECheckBoxState	SRuntimesWidget::IsEntryChecked(int32 Index) const
+{
+	check(Index < CheckBoxNames.Num());
+	const FString& CheckBoxName = CheckBoxNames[Index];
+
+	if (CheckBoxName == TEXT("All"))
+	{
+		return TargetRuntimes.Get().Num() == 0 ? ECheckBoxState::Checked : ECheckBoxState::Unchecked;
+	}
+	
+	return TargetRuntimes.Get().Contains(CheckBoxName) ? ECheckBoxState::Checked : ECheckBoxState::Unchecked;
+}
+
+void SRuntimesWidget::InitCheckBoxNames()
+{
+	CheckBoxNames.Reset();
+
+	CheckBoxNames.Add(TEXT("All"));
+
+	for (TWeakInterfacePtr<INNERuntime> Runtime : UE::NNECore::GetAllRuntimes())
+	{
+		CheckBoxNames.Add(Runtime->GetRuntimeName());
+	}
+
+	//A runtime could possibly be in the target list but not registered to the editor
+	for (const FString& EnabledRuntimeName : TargetRuntimes.Get())
+	{
+		if (CheckBoxNames.Contains(EnabledRuntimeName))
+			continue;
+		CheckBoxNames.Add(EnabledRuntimeName);
+	}
+}
+
+} //UE::NNEEditorTools::Private
+
+#undef LOCTEXT_NAMESPACE
\ No newline at end of file
diff --git a/Engine/Plugins/Experimental/NNE/Source/NNEEditorTools/Private/SNNEEditorToolsRuntimesWidget.h b/Engine/Plugins/Experimental/NNE/Source/NNEEditorTools/Private/SNNEEditorToolsRuntimesWidget.h
new file mode 100644
--- /dev/null
+++ b/Engine/Plugins/Experimental/NNE/Source/NNEEditorTools/Private/SNNEEditorToolsRuntimesWidget.h
@@ -0,0 +1,39 @@
+// Copyright Epic Games, Inc. All Rights Reserved.
+
+#pragma once
+
+#include "CoreMinimal.h"
+#include "Styling/SlateTypes.h"
+#include "Widgets/SCompoundWidget.h"
+
+namespace UE::NNEEditorTools::Private
+{
+
+	class SRuntimesWidget : public SCompoundWidget
+	{
+
+	public:
+
+		DECLARE_DELEGATE_OneParam(FOnTargetRuntimesChanged, TArrayView<const FString> /*NewTargetRuntimes*/)
+
+		SLATE_BEGIN_ARGS(SRuntimesWidget)
+			{}
+			SLATE_ATTRIBUTE(TArrayView<const FString>, TargetRuntimes)
+			SLATE_EVENT(FOnTargetRuntimesChanged, OnTargetRuntimesChanged)
+		SLATE_END_ARGS()
+
+		void Construct(const FArguments& InArgs);
+
+	private:
+		void OnEntryChanged(ECheckBoxState CheckType, int32 Index);
+		ECheckBoxState IsEntryChecked(int32 Index) const;
+
+		void AddCheckBox(const FString& RuntimeName, int32 Index);
+		void InitCheckBoxNames();
+
+		TArray<FString> CheckBoxNames;
+		TAttribute<TArrayView<const FString>> TargetRuntimes;
+		FOnTargetRuntimesChanged OnTargetRuntimesChanged;
+	};
+
+} // UE::NNEEditorTools::Private
\ No newline at end of file
diff --git a/Engine/Plugins/Experimental/NNE/Source/NNEEditorTools/Public/NNEModelDataFactory.h b/Engine/Plugins/Experimental/NNE/Source/NNEEditorTools/Public/NNEModelDataFactory.h
--- a/Engine/Plugins/Experimental/NNE/Source/NNEEditorTools/Public/NNEModelDataFactory.h
+++ b/Engine/Plugins/Experimental/NNE/Source/NNEEditorTools/Public/NNEModelDataFactory.h
@@ -3,6 +3,7 @@
 #pragma once
 
 #include "CoreMinimal.h"
+#include "EditorReimportHandler.h"
 #include "Factories/Factory.h"
 
 #include "NNEModelDataFactory.generated.h"
@@ -13,7 +14,7 @@
  * The currently supported format is .onnx. The factory is invoked when a file of this format is dragged to the editor's content browser.
  */
 UCLASS()
-class NNEEDITORTOOLS_API UNNEModelDataFactory : public UFactory
+class NNEEDITORTOOLS_API UNNEModelDataFactory : public UFactory, public FReimportHandler
 {
 	GENERATED_BODY()
 
@@ -23,7 +24,14 @@ public:
 
 public:
 
+	//~ Begin UFactory Interface
 	virtual UObject* FactoryCreateBinary(UClass* Class, UObject* InParent, FName Name, EObjectFlags Flags, UObject* Context, const TCHAR* Type, const uint8*& Buffer, const uint8* BufferEnd, FFeedbackContext* Warn) override;
 	virtual bool FactoryCanImport(const FString& Filename) override;
+	//~ End UFactory Interface
 
+	//~ Begin FReimportHandler Interface
+	virtual bool CanReimport(UObject* Obj, TArray<FString>& OutFilenames) override;
+	virtual void SetReimportPaths(UObject* Obj, const TArray<FString>& NewReimportPaths) override;
+	virtual EReimportResult::Type Reimport(UObject* Obj) override;
+	//~ End FReimportHandler Interface
 };
\ No newline at end of file
