NNE ModelData only cook target runtimes from import setting.

What:
* The runtime settings are now used by the cooking process, meaning reduction in cooked asset size.

Tests:
* Request disabled runtime in both editor and standalone build from a NNEModelData. Check log and behavior
* Verify cooked asset size in standalone build when only one runtime is enabled.
* NNE smoke tests

#rnx

diff --git a/Engine/Plugins/Experimental/NNE/Source/NNECore/Private/NNECoreModelData.cpp b/Engine/Plugins/Experimental/NNE/Source/NNECore/Private/NNECoreModelData.cpp
--- a/Engine/Plugins/Experimental/NNE/Source/NNECore/Private/NNECoreModelData.cpp
+++ b/Engine/Plugins/Experimental/NNE/Source/NNECore/Private/NNECoreModelData.cpp
@@ -69,7 +69,7 @@ inline void PutIntoDDC(const FGuid& FileDataId, const FString& RuntimeName, FSha
 
 #endif
 
-inline TArray<uint8> Create(const FString& RuntimeName, FString FileType, const TArray<uint8>& FileData)
+inline TArray<uint8> CreateRuntimeDataBlob(const FString& RuntimeName, FString FileType, const TArray<uint8>& FileData)
 {
 	TWeakInterfacePtr<INNERuntime> NNERuntime = UE::NNECore::GetRuntime<INNERuntime>(RuntimeName);
 	if (NNERuntime.IsValid())
@@ -98,6 +98,20 @@ void UNNEModelData::Init(const FString& Type, TConstArrayView<uint8> Buffer)
 
 TConstArrayView<uint8> UNNEModelData::GetModelData(const FString& RuntimeName)
 {
+#if WITH_EDITORONLY_DATA
+	// Check model data is supporting the requested target runtime
+	TArrayView<const FString> TargetRuntimesNames = GetTargetRuntimes();
+	if (!TargetRuntimesNames.IsEmpty() && !TargetRuntimesNames.Contains(RuntimeName))
+	{
+		UE_LOG(LogNNE, Error, TEXT("UNNEModelData: Runtime '%s' is not among the target runtimes. Target runtimes are: "), *RuntimeName);
+		for (const FString& TargetRuntimesName : TargetRuntimesNames)
+		{
+			UE_LOG(LogNNE, Error, TEXT("- %s"), *TargetRuntimesName);
+		}
+		return {};
+	}
+#endif //WITH_EDITORONLY_DATA
+
 	// Check if we have a local cache hit
 	TArray<uint8>* LocalData = ModelData.Find(RuntimeName);
 	if (LocalData)
@@ -115,10 +129,10 @@ TConstArrayView<uint8> UNNEModelData::GetModelData(const FString& RuntimeName)
 		TArray<uint8>* CachedRemoteData = ModelData.Find(RuntimeName);
 		return TConstArrayView<uint8>(CachedRemoteData->GetData(), CachedRemoteData->Num());
 	}
-#endif
-	
+#endif //WITH_EDITOR
+
 	// Try to create the model
-	TArray<uint8> CreatedData = Create(RuntimeName, FileType, FileData);
+	TArray<uint8> CreatedData = CreateRuntimeDataBlob(RuntimeName, FileType, FileData);
 	if (CreatedData.Num() < 1)
 	{
 		return {};
@@ -131,7 +145,7 @@ TConstArrayView<uint8> UNNEModelData::GetModelData(const FString& RuntimeName)
 	// And put it into DDC
 	FSharedBuffer SharedBuffer = MakeSharedBufferFromArray(MoveTemp(CreatedData));
 	PutIntoDDC(FileDataId, RuntimeName, SharedBuffer);
-#endif
+#endif //WITH_EDITOR
 	
 	TArray<uint8>* CachedCreatedData = ModelData.Find(RuntimeName);
 	return TConstArrayView<uint8>(CachedCreatedData->GetData(), CachedCreatedData->Num());
@@ -142,22 +156,34 @@ void UNNEModelData::Serialize(FArchive& Ar)
 	// Store the asset version (no effect in load)
 	Ar.UsingCustomVersion(UNNEModelData::GUID);
 
+#if WITH_EDITORONLY_DATA
 	// Recreate each model data when cooking
 	if (Ar.IsCooking() && Ar.IsSaving())
 	{
 		ModelData.Reset();
-		
-		TArrayView<TWeakInterfacePtr<INNERuntime>> Runtimes = UE::NNECore::GetAllRuntimes();
-		for (int i = 0; i < Runtimes.Num(); i++)
+
+		TArray<FString, TInlineAllocator<10>> CookedRuntimeNames;
+		CookedRuntimeNames.Append(GetTargetRuntimes());
+
+		//No target runtime means all currently registered ones.
+		if (GetTargetRuntimes().IsEmpty())
+		{
+			for (const TWeakInterfacePtr<INNERuntime>& Runtime : UE::NNECore::GetAllRuntimes())
+			{
+				CookedRuntimeNames.Add(Runtime->GetRuntimeName());
+			}
+		}
+
+		for (const FString& RuntimeName : CookedRuntimeNames)
 		{
-			TArray<uint8> CreatedData = Create(Runtimes[i]->GetRuntimeName(), FileType, FileData);
+			TArray<uint8> CreatedData = CreateRuntimeDataBlob(RuntimeName, FileType, FileData);
 			if (CreatedData.Num() > 0)
 			{
-				ModelData.Add(Runtimes[i]->GetRuntimeName(), CreatedData);
+				ModelData.Add(RuntimeName, CreatedData);
 #if WITH_EDITOR
 				FSharedBuffer SharedBuffer = MakeSharedBufferFromArray(MoveTemp(CreatedData));
-				PutIntoDDC(FileDataId, Runtimes[i]->GetRuntimeName(), SharedBuffer);
-#endif
+				PutIntoDDC(FileDataId, RuntimeName, SharedBuffer);
+#endif //WITH_EDITOR
 			}
 		}
 
@@ -179,6 +205,7 @@ void UNNEModelData::Serialize(FArchive& Ar)
 		}
 	}
 	else
+#endif //WITH_EDITORONLY_DATA
 	{
 		int32 NumItems = 0;
 
