NNE Split Model between Model and Model instance

What:
* For runtime to easily implement weight sharing and to allow to spawn a new model without keeping the source asset in memory, NNEModel was split so many instance can be created from a model.

How:
* NNEModel is keeping a copy of its runtime specific data so NNEModelAsset can be safely release after it's creation.
* multiple NNEModelInstance(s) can be spawn from a NNEModel, allowing runtimes to implement weight sharing easily.

Tests:
* All NNE tests green on win64 editor (but shape, known to fail at the moment)
* Style transfer demo on win64 editor

#rnx

diff --git a/Engine/Plugins/AI/MLAdapter/Source/MLAdapter/Private/Agents/MLAdapterAgent_Inference.cpp b/Engine/Plugins/AI/MLAdapter/Source/MLAdapter/Private/Agents/MLAdapterAgent_Inference.cpp
--- a/Engine/Plugins/AI/MLAdapter/Source/MLAdapter/Private/Agents/MLAdapterAgent_Inference.cpp
+++ b/Engine/Plugins/AI/MLAdapter/Source/MLAdapter/Private/Agents/MLAdapterAgent_Inference.cpp
@@ -25,7 +25,7 @@ void UMLAdapterAgent_Inference::PostInitProperties()
 		return;
 	}
 
-	Brain = Runtime->CreateModelCPU(ModelData);
+	Brain = Runtime->CreateModel(ModelData)->CreateModelInstance();
 
 	TConstArrayView<UE::NNECore::FTensorDesc> InputTensorDescs = Brain->GetInputTensorDescs();
 	if (InputTensorDescs.Num() != 1)
diff --git a/Engine/Plugins/AI/MLAdapter/Source/MLAdapter/Public/Agents/MLAdapterAgent_Inference.h b/Engine/Plugins/AI/MLAdapter/Source/MLAdapter/Public/Agents/MLAdapterAgent_Inference.h
--- a/Engine/Plugins/AI/MLAdapter/Source/MLAdapter/Public/Agents/MLAdapterAgent_Inference.h
+++ b/Engine/Plugins/AI/MLAdapter/Source/MLAdapter/Public/Agents/MLAdapterAgent_Inference.h
@@ -30,7 +30,7 @@ public:
 	TObjectPtr<UNNEModelData> ModelData;
 
 	/** The neural network controlling this agent. */
-	TUniquePtr<UE::NNECore::IModelCPU> Brain;
+	TUniquePtr<UE::NNECore::IModelInstanceCPU> Brain;
 
 private:
 	uint32 InputTensorSizeInBytes = 0;
diff --git a/Engine/Plugins/Animation/MLDeformer/VertexDeltaModel/Source/VertexDeltaModel/Private/VertexDeltaGraphDataInterface.cpp b/Engine/Plugins/Animation/MLDeformer/VertexDeltaModel/Source/VertexDeltaModel/Private/VertexDeltaGraphDataInterface.cpp
--- a/Engine/Plugins/Animation/MLDeformer/VertexDeltaModel/Source/VertexDeltaModel/Private/VertexDeltaGraphDataInterface.cpp
+++ b/Engine/Plugins/Animation/MLDeformer/VertexDeltaModel/Source/VertexDeltaModel/Private/VertexDeltaGraphDataInterface.cpp
@@ -115,9 +115,9 @@ namespace UE::VertexDeltaModel
 			if (VertexDeltaModelInstance)
 			{
 				SkeletalMeshObject = ModelInstance->GetSkeletalMeshComponent()->MeshObject;
-				NNEModelRDG = VertexDeltaModelInstance->GetNNEModelRDG();
+				NNEModelInstanceRDG = VertexDeltaModelInstance->GetNNEModelInstanceRDG();
 
-				if (NNEModelRDG)
+				if (NNEModelInstanceRDG)
 				{
 					const FVertexMapBuffer& VertexMapBuffer = VertexDeltaModel->GetVertexMapBuffer();
 					Weight = DeformerComponent->GetWeight();
@@ -134,7 +134,7 @@ namespace UE::VertexDeltaModel
 			return false;
 		}
 
-		if (!bCanRunNeuralNet || SkeletalMeshObject == nullptr || NNEModelRDG == nullptr || VertexMapBufferSRV == nullptr)
+		if (!bCanRunNeuralNet || SkeletalMeshObject == nullptr || NNEModelInstanceRDG == nullptr || VertexMapBufferSRV == nullptr)
 		{
 			return false;
 		}
diff --git a/Engine/Plugins/Animation/MLDeformer/VertexDeltaModel/Source/VertexDeltaModel/Private/VertexDeltaModelInstance.cpp b/Engine/Plugins/Animation/MLDeformer/VertexDeltaModel/Source/VertexDeltaModel/Private/VertexDeltaModelInstance.cpp
--- a/Engine/Plugins/Animation/MLDeformer/VertexDeltaModel/Source/VertexDeltaModel/Private/VertexDeltaModelInstance.cpp
+++ b/Engine/Plugins/Animation/MLDeformer/VertexDeltaModel/Source/VertexDeltaModel/Private/VertexDeltaModelInstance.cpp
@@ -12,9 +12,9 @@ UVertexDeltaModel* UVertexDeltaModelInstance::GetVertexDeltaModel() const
 }
 
 
-UE::NNECore::IModelRDG* UVertexDeltaModelInstance::GetNNEModelRDG() const
+UE::NNECore::IModelInstanceRDG* UVertexDeltaModelInstance::GetNNEModelInstanceRDG() const
 {
-	return ModelRDG.Get();
+	return ModelInstanceRDG.Get();
 }
 
 TRefCountPtr<FRDGPooledBuffer> UVertexDeltaModelInstance::GetOutputRDGBuffer() const
@@ -27,10 +27,10 @@ FString UVertexDeltaModelInstance::CheckCompatibility(USkeletalMeshComponent* In
 	FString ErrorString = Super::CheckCompatibility(InSkelMeshComponent, bLogIssues);
 
 	// Verify the number of inputs versus the expected number of inputs.	
-	const UE::NNECore::IModelRDG* RDGModel = GetNNEModelRDG();
-	if (RDGModel && Model->GetDeformerAsset())
+	const UE::NNECore::IModelInstanceRDG* RDGModelInstance = GetNNEModelInstanceRDG();
+	if (RDGModelInstance && Model->GetDeformerAsset())
 	{
-		TConstArrayView<UE::NNECore::FTensorDesc> GPUTensorDesc = RDGModel->GetInputTensorDescs();
+		TConstArrayView<UE::NNECore::FTensorDesc> GPUTensorDesc = RDGModelInstance->GetInputTensorDescs();
 		const int64 NumNeuralNetInputs = GPUTensorDesc[0].GetShape().Rank();
 		const int64 NumDeformerAssetInputs = static_cast<int64>(Model->GetInputInfo()->CalcNumNeuralNetInputs(Model->GetNumFloatsPerBone(), Model->GetNumFloatsPerCurve()));
 		if (NumNeuralNetInputs != NumDeformerAssetInputs)
@@ -49,14 +49,14 @@ FString UVertexDeltaModelInstance::CheckCompatibility(USkeletalMeshComponent* In
 
 bool UVertexDeltaModelInstance::IsValidForDataProvider() const
 {
-	return ModelRDG.IsValid(); 
+	return ModelInstanceRDG.IsValid(); 
 }
 
 void UVertexDeltaModelInstance::Execute(float ModelWeight)
 {
 	TRACE_CPUPROFILER_EVENT_SCOPE(UMLDeformerModelInstance::Execute)
 
-	if (ModelRDG)
+	if (ModelInstanceRDG)
 	{
 		ENQUEUE_RENDER_COMMAND(RunNeuralNetwork)
 			(
@@ -82,11 +82,11 @@ void UVertexDeltaModelInstance::Execute(float ModelWeight)
 					UE::NNECore::FTensorBindingRDG& OutputBindingRDG = OutputBindingsRDG.Emplace_GetRef();
 					OutputBindingRDG.Buffer = GraphBuilder.RegisterExternalBuffer(RDGVertexDeltaBuffer);
 
-					if (!ModelRDG.IsValid())
+					if (!ModelInstanceRDG.IsValid())
 					{
 						return;
 					}
-					int32 Res = ModelRDG->EnqueueRDG(GraphBuilder, InputBindingsRDG, OutputBindingsRDG);
+					int32 Res = ModelInstanceRDG->EnqueueRDG(GraphBuilder, InputBindingsRDG, OutputBindingsRDG);
 	
 					GraphBuilder.Execute();
 				}
@@ -157,7 +157,7 @@ void UVertexDeltaModelInstance::PostMLDeformerComponentInit()
 
 void UVertexDeltaModelInstance::CreateNNEModel()
 {
-	if (!ModelRDG.IsValid())
+	if (!ModelInstanceRDG.IsValid())
 	{
 		UVertexDeltaModel* VertexDeltaModel = Cast<UVertexDeltaModel>(Model);
 		if (VertexDeltaModel)
@@ -177,19 +177,19 @@ void UVertexDeltaModelInstance::CreateNNEModel()
 				if (RuntimeRDG.IsValid())
 				{
 					// allocate tensor inputs and outputs
-					ModelRDG = RuntimeRDG->CreateModelRDG(ModelData);
+					ModelInstanceRDG = RuntimeRDG->CreateModel(ModelData)->CreateModelInstance();
 
-					if (ModelRDG)
+					if (ModelInstanceRDG)
 					{
 						// setup inputs
-						TConstArrayView<UE::NNECore::FTensorDesc> InputTensorDescs = ModelRDG->GetInputTensorDescs();
+						TConstArrayView<UE::NNECore::FTensorDesc> InputTensorDescs = ModelInstanceRDG->GetInputTensorDescs();
 						UE::NNECore::FTensorShape InputTensorShape = UE::NNECore::FTensorShape::MakeFromSymbolic(InputTensorDescs[0].GetShape());
-						ModelRDG->SetInputTensorShapes({ InputTensorShape });
+						ModelInstanceRDG->SetInputTensorShapes({ InputTensorShape });
 						check(InputTensorDescs[0].GetElemByteSize() == sizeof(float));
 						NNEInputTensorBuffer.SetNumUninitialized(InputTensorShape.Volume());
 
 						// setup outputs
-						TConstArrayView<UE::NNECore::FTensorDesc> OutputTensorDescs = ModelRDG->GetOutputTensorDescs();
+						TConstArrayView<UE::NNECore::FTensorDesc> OutputTensorDescs = ModelInstanceRDG->GetOutputTensorDescs();
 						CreateRDGBuffers(OutputTensorDescs);
 					}
 					else
@@ -218,12 +218,12 @@ bool UVertexDeltaModelInstance::SetupInputs()
 
 	// Get the network and make sure it's loaded.
 
-	if (!ModelRDG)
+	if (!ModelInstanceRDG)
 	{
 		return false;
 	}
 
-	TConstArrayView<UE::NNECore::FTensorDesc> InputTensorDescs = ModelRDG->GetInputTensorDescs();
+	TConstArrayView<UE::NNECore::FTensorDesc> InputTensorDescs = ModelInstanceRDG->GetInputTensorDescs();
 	const UE::NNECore::FTensorShape InputTensorShape = UE::NNECore::FTensorShape::MakeFromSymbolic(InputTensorDescs[0].GetShape());
 	const int64 NumNeuralNetInputs = InputTensorShape.Volume();
 	const int64 NumDeformerAssetInputs = Model->GetInputInfo()->CalcNumNeuralNetInputs(Model->GetNumFloatsPerBone(), Model->GetNumFloatsPerCurve());
diff --git a/Engine/Plugins/Animation/MLDeformer/VertexDeltaModel/Source/VertexDeltaModel/Public/VertexDeltaGraphDataInterface.h b/Engine/Plugins/Animation/MLDeformer/VertexDeltaModel/Source/VertexDeltaModel/Public/VertexDeltaGraphDataInterface.h
--- a/Engine/Plugins/Animation/MLDeformer/VertexDeltaModel/Source/VertexDeltaModel/Public/VertexDeltaGraphDataInterface.h
+++ b/Engine/Plugins/Animation/MLDeformer/VertexDeltaModel/Source/VertexDeltaModel/Public/VertexDeltaGraphDataInterface.h
@@ -85,7 +85,7 @@ namespace UE::VertexDeltaModel
 
 	protected:
 		FSkeletalMeshObject* SkeletalMeshObject = nullptr;
-		UE::NNECore::IModelRDG* NNEModelRDG = nullptr;
+		UE::NNECore::IModelInstanceRDG* NNEModelInstanceRDG = nullptr;
 		UVertexDeltaModelInstance* VertexDeltaModelInstance = nullptr;
 		FRHIShaderResourceView* VertexMapBufferSRV = nullptr;
 		FRDGBuffer* Buffer = nullptr;
diff --git a/Engine/Plugins/Animation/MLDeformer/VertexDeltaModel/Source/VertexDeltaModel/Public/VertexDeltaModelInstance.h b/Engine/Plugins/Animation/MLDeformer/VertexDeltaModel/Source/VertexDeltaModel/Public/VertexDeltaModelInstance.h
--- a/Engine/Plugins/Animation/MLDeformer/VertexDeltaModel/Source/VertexDeltaModel/Public/VertexDeltaModelInstance.h
+++ b/Engine/Plugins/Animation/MLDeformer/VertexDeltaModel/Source/VertexDeltaModel/Public/VertexDeltaModelInstance.h
@@ -31,10 +31,10 @@ public:
 	virtual void PostMLDeformerComponentInit() override;
 
 	/**
-	 * Get the RDG Neural Network Model
-	 * @return The IModelRDG if successfully created, or nullptr
+	 * Get the RDG Neural Network Model Instance
+	 * @return The IModelInstanceRDG if successfully created, or nullptr
 	 */
-	UE::NNECore::IModelRDG* GetNNEModelRDG() const;
+	UE::NNECore::IModelInstanceRDG* GetNNEModelInstanceRDG() const;
 
 	/**
 	 * Get the output vertex delta buffer
@@ -44,7 +44,7 @@ public:
 
 	/**
 	 * Get the render graph buffer description required for the output of a neural network. Return false if a flat buffer is not appropriate
-	 * @return True if the OutputTensorDescs can be rppresented by a flat float  FRDGBuffer
+	 * @return True if the OutputTensorDescs can be represented by a flat float  FRDGBuffer
 	 */
 	bool GetRDGVertexBufferDesc(TConstArrayView<UE::NNECore::FTensorDesc>& InOutputTensorDescs, FRDGBufferDesc& OutBufferDesc);
 
@@ -60,7 +60,7 @@ private:
 	TRefCountPtr<FRDGPooledBuffer> RDGVertexDeltaBuffer;
 
 	// The NNE RDG Model 
-	TUniquePtr<UE::NNECore::IModelRDG> ModelRDG;
+	TUniquePtr<UE::NNECore::IModelInstanceRDG> ModelInstanceRDG;
 	
 	// The CPU Input Tensor Buffer
 	TArray<float> NNEInputTensorBuffer;
diff --git a/Engine/Plugins/Experimental/NNE/Source/NNECore/Internal/NNECoreModelBase.h b/Engine/Plugins/Experimental/NNE/Source/NNECore/Internal/NNECoreModelBase.h
--- a/Engine/Plugins/Experimental/NNE/Source/NNECore/Internal/NNECoreModelBase.h
+++ b/Engine/Plugins/Experimental/NNE/Source/NNECore/Internal/NNECoreModelBase.h
@@ -8,11 +8,11 @@
 
 namespace UE::NNECore::Internal
 {
-	template <class ModelInterface> class FModelBase : public ModelInterface
+	template <class ModelInterface> class FModelInstanceBase : public ModelInterface
 	{
 	public:
 
-		virtual ~FModelBase() = default;
+		virtual ~FModelInstanceBase() = default;
 
 		virtual TConstArrayView<NNECore::FTensorDesc> GetInputTensorDescs() const override;
 		virtual TConstArrayView<NNECore::FTensorDesc> GetOutputTensorDescs() const override;
@@ -22,7 +22,7 @@ namespace UE::NNECore::Internal
 
 	protected:
 
-		FModelBase() {}
+		FModelInstanceBase() {}
 
 		TArray<FTensorShape>	InputTensorShapes;
 		TArray<FTensorShape>	OutputTensorShapes;
@@ -30,27 +30,27 @@ namespace UE::NNECore::Internal
 		TArray<FTensorDesc>		OutputSymbolicTensors;
 	};
 
-	template <class T> TConstArrayView<FTensorDesc> FModelBase<T>::GetInputTensorDescs() const
+	template <class T> TConstArrayView<FTensorDesc> FModelInstanceBase<T>::GetInputTensorDescs() const
 	{
 		return InputSymbolicTensors;
 	}
 
-	template <class T> TConstArrayView<FTensorDesc> FModelBase<T>::GetOutputTensorDescs() const
+	template <class T> TConstArrayView<FTensorDesc> FModelInstanceBase<T>::GetOutputTensorDescs() const
 	{
 		return OutputSymbolicTensors;
 	}
 
-	template <class T> TConstArrayView<FTensorShape> FModelBase<T>::GetInputTensorShapes() const
+	template <class T> TConstArrayView<FTensorShape> FModelInstanceBase<T>::GetInputTensorShapes() const
 	{
 		return InputTensorShapes;
 	}
 
-	template <class T> TConstArrayView<FTensorShape> FModelBase<T>::GetOutputTensorShapes() const
+	template <class T> TConstArrayView<FTensorShape> FModelInstanceBase<T>::GetOutputTensorShapes() const
 	{
 		return OutputTensorShapes;
 	}
 
-	template <class T> int FModelBase<T>::SetInputTensorShapes(TConstArrayView<FTensorShape> InInputShapes)
+	template <class T> int FModelInstanceBase<T>::SetInputTensorShapes(TConstArrayView<FTensorShape> InInputShapes)
 	{
 		InputTensorShapes.Reset(InInputShapes.Num());
 
diff --git a/Engine/Plugins/Experimental/NNE/Source/NNECore/Public/NNECoreRuntimeCPU.h b/Engine/Plugins/Experimental/NNE/Source/NNECore/Public/NNECoreRuntimeCPU.h
--- a/Engine/Plugins/Experimental/NNE/Source/NNECore/Public/NNECoreRuntimeCPU.h
+++ b/Engine/Plugins/Experimental/NNE/Source/NNECore/Public/NNECoreRuntimeCPU.h
@@ -26,15 +26,16 @@ struct NNECORE_API FTensorBindingCPU
 };
 
 /**
- * The interface of a model that can run on CPU.
+ * The interface of a model instance that can run on CPU.
  *
+ * Use UE::NNECore::IModelCPU::CreateModelInstance() to get a model instance.
  * Use UE::NNECore::GetRuntime<INNERuntimeCPU>(RuntimeName) to get a runtime capable of creating CPU models.
  */
-class NNECORE_API IModelCPU
+class NNECORE_API IModelInstanceCPU
 {
 public:
 
-	virtual ~IModelCPU() = default;
+	virtual ~IModelInstanceCPU() = default;
 
 	/**
 	 * Get the input tensor descriptions as defined by the model, potentially with variable dimensions.
@@ -96,6 +97,28 @@ public:
 	virtual int RunSync(TConstArrayView<FTensorBindingCPU> InInputTensors, TConstArrayView<FTensorBindingCPU> InOutputTensors) = 0;
 };
 
+/**
+ * The interface of a model capable of creating model instance that can run on CPU.
+ *
+ * Use UE::NNECore::GetRuntime<INNERuntimeCPU>(RuntimeName) to get a runtime capable of creating CPU models.
+ */
+class NNECORE_API IModelCPU
+{
+public:
+
+	virtual ~IModelCPU() = default;
+
+	/**
+	 * Create a model instance for inference
+	 *
+	 * The runtime have the opportunity to share the model weights among multiple IModelInstanceCPU created from an IModelCPU instance, however this is not mandatory.
+	 * The caller can decide to convert the result into a shared pointer if required (e.g. if the model needs to be shared with an async task for evaluation).
+	 *
+	 * @return A caller owned model representing the neural network instance created.
+	 */
+	virtual TUniquePtr<UE::NNECore::IModelInstanceCPU> CreateModelInstance() = 0;
+};
+
 } // UE::NNECore
 
 UINTERFACE()
@@ -133,5 +156,5 @@ public:
 	 * @param ModelData The model data for which to create a model.
 	 * @return A caller owned model representing the neural network created from ModelData.
 	 */
-	virtual TUniquePtr<UE::NNECore::IModelCPU> CreateModelCPU(TObjectPtr<UNNEModelData> ModelData) = 0;
+	virtual TUniquePtr<UE::NNECore::IModelCPU> CreateModel(TObjectPtr<UNNEModelData> ModelData) = 0;
 };
\ No newline at end of file
diff --git a/Engine/Plugins/Experimental/NNE/Source/NNECore/Public/NNECoreRuntimeGPU.h b/Engine/Plugins/Experimental/NNE/Source/NNECore/Public/NNECoreRuntimeGPU.h
--- a/Engine/Plugins/Experimental/NNE/Source/NNECore/Public/NNECoreRuntimeGPU.h
+++ b/Engine/Plugins/Experimental/NNE/Source/NNECore/Public/NNECoreRuntimeGPU.h
@@ -26,15 +26,16 @@ struct NNECORE_API FTensorBindingGPU
 };
 
 /**
- * The interface of a model that can run on GPU.
+ * The interface of a model instance that can run on GPU.
  *
+ * Use UE::NNECore::IModelGPU::CreateModelInstance() to get a model instance.
  * Use UE::NNECore::GetRuntime<INNERuntimeGPU>(RuntimeName) to get a runtime capable of creating GPU models.
  */
-class NNECORE_API IModelGPU
+class NNECORE_API IModelInstanceGPU
 {
 public:
 
-	virtual ~IModelGPU() = default;
+	virtual ~IModelInstanceGPU() = default;
 
 	/**
 	 * Get the input tensor descriptions as defined by the model, potentially with variable dimensions.
@@ -96,6 +97,28 @@ public:
 	virtual int RunSync(TConstArrayView<FTensorBindingGPU> InInputTensors, TConstArrayView<FTensorBindingGPU> InOutputTensors) = 0;
 };
 
+/**
+ * The interface of a model capable of creating model instance that can run on GPU.
+ *
+ * Use UE::NNECore::GetRuntime<INNERuntimeGPU>(RuntimeName) to get a runtime capable of creating GPU models.
+ */
+class NNECORE_API IModelGPU
+{
+public:
+
+	virtual ~IModelGPU() = default;
+
+	/**
+	 * Create a model instance for inference
+	 *
+	 * The runtime have the opportunity to share the model weights among multiple IModelInstanceGPU created from an IModelGPU instance, however this is not mandatory.
+	 * The caller can decide to convert the result into a shared pointer if required (e.g. if the model needs to be shared with an async task for evaluation).
+	 *
+	 * @return A caller owned model representing the neural network instance created.
+	 */
+	virtual TUniquePtr<UE::NNECore::IModelInstanceGPU> CreateModelInstance() = 0;
+};
+
 } // UE::NNECore
 
 UINTERFACE()
@@ -133,5 +156,5 @@ public:
 	 * @param ModelData The model data for which to create a model.
 	 * @return A caller owned model representing the neural network created from ModelData.
 	 */
-	virtual TUniquePtr<UE::NNECore::IModelGPU> CreateModelGPU(TObjectPtr<UNNEModelData> ModelData) = 0;
+	virtual TUniquePtr<UE::NNECore::IModelGPU> CreateModel(TObjectPtr<UNNEModelData> ModelData) = 0;
 };
\ No newline at end of file
diff --git a/Engine/Plugins/Experimental/NNE/Source/NNECore/Public/NNECoreRuntimeRDG.h b/Engine/Plugins/Experimental/NNE/Source/NNECore/Public/NNECoreRuntimeRDG.h
--- a/Engine/Plugins/Experimental/NNE/Source/NNECore/Public/NNECoreRuntimeRDG.h
+++ b/Engine/Plugins/Experimental/NNE/Source/NNECore/Public/NNECoreRuntimeRDG.h
@@ -25,15 +25,16 @@ struct NNECORE_API FTensorBindingRDG
 };
 
 /**
- * The interface of a model that can run on RDG.
+ * The interface of a model instance that can run on RDG.
  *
+ * Use UE::NNECore::IModelRDG::CreateModelInstance() to get a model instance.
  * Use UE::NNECore::GetRuntime<INNERuntimeRDG>(RuntimeName) to get a runtime capable of creating RDG models.
  */
-class NNECORE_API IModelRDG
+class NNECORE_API IModelInstanceRDG
 {
 public:
 
-	virtual ~IModelRDG() = default;
+	virtual ~IModelInstanceRDG() = default;
 
 	/**
 	 * Get the input tensor descriptions as defined by the model, potentially with variable dimensions.
@@ -96,6 +97,28 @@ public:
 	virtual int EnqueueRDG(FRDGBuilder& RDGBuilder, TConstArrayView<FTensorBindingRDG> Inputs, TConstArrayView<FTensorBindingRDG> Outputs) = 0;
 };
 
+/**
+ * The interface of a model capable of creating model instance that can run on RDG.
+ *
+ * Use UE::NNECore::GetRuntime<INNERuntimeRDG>(RuntimeName) to get a runtime capable of creating RDG models.
+ */
+class NNECORE_API IModelRDG
+{
+public:
+
+	virtual ~IModelRDG() = default;
+
+	/**
+	 * Create a model instance for inference
+	 *
+	 * The runtime have the opportunity to share the model weights among multiple IModelInstanceRDG created from an IModelRDG instance, however this is not mandatory.
+	 * The caller can decide to convert the result into a shared pointer if required (e.g. if the model needs to be shared with an async task for evaluation).
+	 *
+	 * @return A caller owned model representing the neural network instance created.
+	 */
+	virtual TUniquePtr<UE::NNECore::IModelInstanceRDG> CreateModelInstance() = 0;
+};
+
 } // UE::NNECore
 
 UINTERFACE()
@@ -133,5 +156,5 @@ public:
 	 * @param ModelData The model data for which to create a model.
 	 * @return A caller owned model representing the neural network created from ModelData.
 	 */
-	virtual TUniquePtr<UE::NNECore::IModelRDG> CreateModelRDG(TObjectPtr<UNNEModelData> ModelData) = 0;
+	virtual TUniquePtr<UE::NNECore::IModelRDG> CreateModel(TObjectPtr<UNNEModelData> ModelData) = 0;
 };
\ No newline at end of file
diff --git a/Engine/Plugins/Experimental/NNE/Source/NNEQA/Private/NNEQAModel.cpp b/Engine/Plugins/Experimental/NNE/Source/NNEQA/Private/NNEQAModel.cpp
--- a/Engine/Plugins/Experimental/NNE/Source/NNEQA/Private/NNEQAModel.cpp
+++ b/Engine/Plugins/Experimental/NNE/Source/NNEQA/Private/NNEQAModel.cpp
@@ -129,60 +129,60 @@ namespace FModelQAHelpers
 	
 TConstArrayView<NNECore::FTensorDesc> FModelQA::GetInputTensorDescs() const
 { 
-	if (ModelCPU.IsValid()) 
-		return ModelCPU->GetInputTensorDescs();
-	if (ModelGPU.IsValid())
-		return ModelGPU->GetInputTensorDescs();
-	if (ModelRDG.IsValid())
-		return ModelRDG->GetInputTensorDescs();
+	if (ModelInstanceCPU.IsValid()) 
+		return ModelInstanceCPU->GetInputTensorDescs();
+	if (ModelInstanceGPU.IsValid())
+		return ModelInstanceGPU->GetInputTensorDescs();
+	if (ModelInstanceRDG.IsValid())
+		return ModelInstanceRDG->GetInputTensorDescs();
 	else
 		return {};
 };
 		
 TConstArrayView<NNECore::FTensorDesc> FModelQA::GetOutputTensorDescs() const
 {
-	if (ModelCPU.IsValid())
-		return ModelCPU->GetOutputTensorDescs();
-	if (ModelGPU.IsValid())
-		return ModelGPU->GetOutputTensorDescs();
-	if (ModelRDG.IsValid())
-		return ModelRDG->GetOutputTensorDescs();
+	if (ModelInstanceCPU.IsValid())
+		return ModelInstanceCPU->GetOutputTensorDescs();
+	if (ModelInstanceGPU.IsValid())
+		return ModelInstanceGPU->GetOutputTensorDescs();
+	if (ModelInstanceRDG.IsValid())
+		return ModelInstanceRDG->GetOutputTensorDescs();
 	else
 		return {};
 }
 
 TConstArrayView<NNECore::FTensorShape> FModelQA::GetInputTensorShapes() const
 {
-	if (ModelCPU.IsValid())
-		return ModelCPU->GetInputTensorShapes();
-	if (ModelGPU.IsValid())
-		return ModelGPU->GetInputTensorShapes();
-	if (ModelRDG.IsValid())
-		return ModelRDG->GetInputTensorShapes();
+	if (ModelInstanceCPU.IsValid())
+		return ModelInstanceCPU->GetInputTensorShapes();
+	if (ModelInstanceGPU.IsValid())
+		return ModelInstanceGPU->GetInputTensorShapes();
+	if (ModelInstanceRDG.IsValid())
+		return ModelInstanceRDG->GetInputTensorShapes();
 	else
 		return {};
 }
 
 TConstArrayView<NNECore::FTensorShape> FModelQA::GetOutputTensorShapes() const
 {
-	if (ModelCPU.IsValid())
-		return ModelCPU->GetOutputTensorShapes();
-	if (ModelGPU.IsValid())
-		return ModelGPU->GetOutputTensorShapes();
-	if (ModelRDG.IsValid())
-		return ModelRDG->GetOutputTensorShapes();
+	if (ModelInstanceCPU.IsValid())
+		return ModelInstanceCPU->GetOutputTensorShapes();
+	if (ModelInstanceGPU.IsValid())
+		return ModelInstanceGPU->GetOutputTensorShapes();
+	if (ModelInstanceRDG.IsValid())
+		return ModelInstanceRDG->GetOutputTensorShapes();
 	else
 		return {};
 }
 
 int FModelQA::SetInputTensorShapes(TConstArrayView<NNECore::FTensorShape> InInputShapes)
 {
-	if (ModelCPU.IsValid())
-		return ModelCPU->SetInputTensorShapes(InInputShapes);
-	if (ModelGPU.IsValid())
-		return ModelGPU->SetInputTensorShapes(InInputShapes);
-	if (ModelRDG.IsValid())
-		return ModelRDG->SetInputTensorShapes(InInputShapes);
+	if (ModelInstanceCPU.IsValid())
+		return ModelInstanceCPU->SetInputTensorShapes(InInputShapes);
+	if (ModelInstanceGPU.IsValid())
+		return ModelInstanceGPU->SetInputTensorShapes(InInputShapes);
+	if (ModelInstanceRDG.IsValid())
+		return ModelInstanceRDG->SetInputTensorShapes(InInputShapes);
 	else
 		return {};
 }
@@ -193,7 +193,7 @@ int FModelQA::EnqueueRDG(FRDGBuilder& RDGBuilder,
 {
 	check(IsInRenderingThread());
 			
-	if (!ModelRDG.IsValid())
+	if (!ModelInstanceRDG.IsValid())
 	{
 		return -1;
 	}
@@ -201,12 +201,12 @@ int FModelQA::EnqueueRDG(FRDGBuilder& RDGBuilder,
 	TArray<NNECore::FTensorBindingRDG> InputBindingsRDG;
 	TArray<NNECore::FTensorBindingRDG> OutputBindingsRDG;
 			
-	FModelQAHelpers::ConvertBinding(RDGBuilder, ModelRDG->GetInputTensorDescs(), InInputBindings, InputBindingsRDG, true);
-	FModelQAHelpers::ConvertBinding(RDGBuilder, ModelRDG->GetOutputTensorDescs(), InOutputBindings, OutputBindingsRDG, false);
+	FModelQAHelpers::ConvertBinding(RDGBuilder, ModelInstanceRDG->GetInputTensorDescs(), InInputBindings, InputBindingsRDG, true);
+	FModelQAHelpers::ConvertBinding(RDGBuilder, ModelInstanceRDG->GetOutputTensorDescs(), InOutputBindings, OutputBindingsRDG, false);
 			
 	FModelQAHelpers::UploadsBindingToGPU(RDGBuilder, InInputBindings, InputBindingsRDG);
 
-	int32 Res = ModelRDG->EnqueueRDG(RDGBuilder, InputBindingsRDG, OutputBindingsRDG);
+	int32 Res = ModelInstanceRDG->EnqueueRDG(RDGBuilder, InputBindingsRDG, OutputBindingsRDG);
 
 	FModelQAHelpers::DownloadBindingToCPU(RDGBuilder, Readbacks, OutputBindingsRDG, InOutputBindings);
 			
@@ -215,19 +215,19 @@ int FModelQA::EnqueueRDG(FRDGBuilder& RDGBuilder,
 
 int FModelQA::RunSync(TConstArrayView<NNECore::FTensorBindingCPU> InInputBindings, TConstArrayView<NNECore::FTensorBindingCPU> InOutputBindings)
 {
-	if (ModelCPU.IsValid())
+	if (ModelInstanceCPU.IsValid())
 	{
-		return ModelCPU->RunSync(InInputBindings, InOutputBindings);
+		return ModelInstanceCPU->RunSync(InInputBindings, InOutputBindings);
 	}
-	if (ModelGPU.IsValid())
+	if (ModelInstanceGPU.IsValid())
 	{
 		TArray<NNECore::FTensorBindingGPU> InputBindingsGPU;
 		TArray<NNECore::FTensorBindingGPU> OutputBindingsGPU;
 		FModelQAHelpers::ConvertBinding(InInputBindings, InputBindingsGPU);
 		FModelQAHelpers::ConvertBinding(InOutputBindings, OutputBindingsGPU);
-		return ModelGPU->RunSync(InputBindingsGPU, OutputBindingsGPU);
+		return ModelInstanceGPU->RunSync(InputBindingsGPU, OutputBindingsGPU);
 	}
-	else if (ModelRDG.IsValid())
+	else if (ModelInstanceRDG.IsValid())
 	{
 		//Note: ok if pending?
 		Readbacks.Empty();
@@ -300,15 +300,15 @@ TUniquePtr<FModelQA> FModelQA::MakeModelQA(const FNNEModelRaw& ONNXModelData, co
 	}
 	if (RuntimeCPU.IsValid())
 	{
-		ModelQA->ModelCPU = RuntimeCPU->CreateModelCPU(ModelData);
+		ModelQA->ModelInstanceCPU = RuntimeCPU->CreateModel(ModelData)->CreateModelInstance();
 	}
 	else if (RuntimeGPU.IsValid())
 	{
-		ModelQA->ModelGPU = RuntimeGPU->CreateModelGPU(ModelData);
+		ModelQA->ModelInstanceGPU = RuntimeGPU->CreateModel(ModelData)->CreateModelInstance();
 	}
 	else if (RuntimeRDG.IsValid())
 	{
-		ModelQA->ModelRDG = RuntimeRDG->CreateModelRDG(ModelData);
+		ModelQA->ModelInstanceRDG = RuntimeRDG->CreateModel(ModelData)->CreateModelInstance();
 	}
 	else
 	{
diff --git a/Engine/Plugins/Experimental/NNE/Source/NNEQA/Private/NNEQAModel.h b/Engine/Plugins/Experimental/NNE/Source/NNEQA/Private/NNEQAModel.h
--- a/Engine/Plugins/Experimental/NNE/Source/NNEQA/Private/NNEQAModel.h
+++ b/Engine/Plugins/Experimental/NNE/Source/NNEQA/Private/NNEQAModel.h
@@ -17,15 +17,15 @@
 namespace UE::NNEQA::Private 
 {
 
-// This class is a wrapper around a NNECore::IModelCPU and NNECore::IModelRDG
-// Allowing to run a model on CPU or GPU using IModelCPU interface only.
-class FModelQA : NNECore::IModelCPU
+// This class is a wrapper around a NNECore::IModelInstanceCPU and NNECore::IModelInstanceRDG
+// Allowing to run a model on CPU or GPU using IModelInstanceCPU interface only.
+class FModelQA : NNECore::IModelInstanceCPU
 {
 private:
 
-	TUniquePtr<NNECore::IModelCPU> ModelCPU;
-	TUniquePtr<NNECore::IModelGPU> ModelGPU;
-	TUniquePtr<NNECore::IModelRDG> ModelRDG;
+	TUniquePtr<NNECore::IModelInstanceCPU> ModelInstanceCPU;
+	TUniquePtr<NNECore::IModelInstanceGPU> ModelInstanceGPU;
+	TUniquePtr<NNECore::IModelInstanceRDG> ModelInstanceRDG;
 
 public:
 	struct FReadbackEntry
diff --git a/Engine/Plugins/Experimental/NNE/Source/NNERuntimeORTCpu/Private/NNERuntimeORTCpu.cpp b/Engine/Plugins/Experimental/NNE/Source/NNERuntimeORTCpu/Private/NNERuntimeORTCpu.cpp
--- a/Engine/Plugins/Experimental/NNE/Source/NNERuntimeORTCpu/Private/NNERuntimeORTCpu.cpp
+++ b/Engine/Plugins/Experimental/NNE/Source/NNERuntimeORTCpu/Private/NNERuntimeORTCpu.cpp
@@ -65,7 +65,7 @@ bool UNNERuntimeORTCpuImpl::CanCreateModelCPU(TObjectPtr<UNNEModelData> ModelDat
 	return bResult;
 }
 
-TUniquePtr<UE::NNECore::IModelCPU> UNNERuntimeORTCpuImpl::CreateModelCPU(TObjectPtr<UNNEModelData> ModelData)
+TUniquePtr<UE::NNECore::IModelCPU> UNNERuntimeORTCpuImpl::CreateModel(TObjectPtr<UNNEModelData> ModelData)
 {
 	check(ModelData != nullptr);
 	
@@ -74,15 +74,8 @@ TUniquePtr<UE::NNECore::IModelCPU> UNNERuntimeORTCpuImpl::CreateModelCPU(TObject
 		return TUniquePtr<UE::NNECore::IModelCPU>();
 	}
 
-	const UE::NNERuntimeORTCpu::Private::FRuntimeConf InConf;
-	UE::NNERuntimeORTCpu::Private::FModelCPU* Model = new UE::NNERuntimeORTCpu::Private::FModelCPU(&NNEEnvironmentCPU, InConf);
 	TConstArrayView<uint8> Data = ModelData->GetModelData(GetRuntimeName());
-	
-	if (!Model->Init(Data))
-	{
-		delete Model;
-		return TUniquePtr<UE::NNECore::IModelCPU>();
-	}
+	UE::NNERuntimeORTCpu::Private::FModelCPU* Model = new UE::NNERuntimeORTCpu::Private::FModelCPU(&NNEEnvironmentCPU, Data);
 	UE::NNECore::IModelCPU* IModel = static_cast<UE::NNECore::IModelCPU*>(Model);
 	return TUniquePtr<UE::NNECore::IModelCPU>(IModel);
 }
\ No newline at end of file
diff --git a/Engine/Plugins/Experimental/NNE/Source/NNERuntimeORTCpu/Private/NNERuntimeORTCpu.h b/Engine/Plugins/Experimental/NNE/Source/NNERuntimeORTCpu/Private/NNERuntimeORTCpu.h
--- a/Engine/Plugins/Experimental/NNE/Source/NNERuntimeORTCpu/Private/NNERuntimeORTCpu.h
+++ b/Engine/Plugins/Experimental/NNE/Source/NNERuntimeORTCpu/Private/NNERuntimeORTCpu.h
@@ -38,5 +38,5 @@ public:
 	virtual TArray<uint8> CreateModelData(FString FileType, TConstArrayView<uint8> FileData) override;
 
 	virtual bool CanCreateModelCPU(TObjectPtr<UNNEModelData> ModelData) const override;
-	virtual TUniquePtr<UE::NNECore::IModelCPU> CreateModelCPU(TObjectPtr<UNNEModelData> ModelData) override;
+	virtual TUniquePtr<UE::NNECore::IModelCPU> CreateModel(TObjectPtr<UNNEModelData> ModelData) override;
 };
\ No newline at end of file
diff --git a/Engine/Plugins/Experimental/NNE/Source/NNERuntimeORTCpu/Private/NNERuntimeORTCpuModel.cpp b/Engine/Plugins/Experimental/NNE/Source/NNERuntimeORTCpu/Private/NNERuntimeORTCpuModel.cpp
--- a/Engine/Plugins/Experimental/NNE/Source/NNERuntimeORTCpu/Private/NNERuntimeORTCpuModel.cpp
+++ b/Engine/Plugins/Experimental/NNE/Source/NNERuntimeORTCpu/Private/NNERuntimeORTCpuModel.cpp
@@ -7,12 +7,12 @@
 namespace UE::NNERuntimeORTCpu::Private
 {
 
-	FModelCPU::FModelCPU() :
+	FModelInstanceCPU::FModelInstanceCPU() :
 		bIsLoaded(false),
 		bHasRun(false)
 	{ }
 
-	FModelCPU::FModelCPU(
+	FModelInstanceCPU::FModelInstanceCPU(
 		Ort::Env* InORTEnvironment,
 		const FRuntimeConf& InRuntimeConf) :
 		bIsLoaded(false),
@@ -21,9 +21,9 @@ namespace UE::NNERuntimeORTCpu::Private
 		ORTEnvironment(InORTEnvironment)
 	{ }
 
-	bool FModelCPU::Init(TConstArrayView<uint8> ModelData)
+	bool FModelInstanceCPU::Init(TConstArrayView<uint8> ModelData)
 	{
-		DECLARE_SCOPE_CYCLE_COUNTER(TEXT("FModelCPU_Init"), STAT_FModelCPU_Init, STATGROUP_NNE);
+		DECLARE_SCOPE_CYCLE_COUNTER(TEXT("FModelInstanceCPU_Init"), STAT_FModelInstanceCPU_Init, STATGROUP_NNE);
 		
 		// Get the header size
 		int32 GuidSize = sizeof(UNNERuntimeORTCpuImpl::GUID);
@@ -37,7 +37,7 @@ namespace UE::NNERuntimeORTCpu::Private
 		{
 			if (ModelBuffer.Num() == 0)
 			{
-				UE_LOG(LogNNE, Warning, TEXT("FModelCPU::Load(): Input model data is empty."));
+				UE_LOG(LogNNE, Warning, TEXT("FModelInstanceCPU::Load(): Input model data is empty."));
 				return false;
 			}
 
@@ -54,7 +54,7 @@ namespace UE::NNERuntimeORTCpu::Private
 			}
 
 			{
-				DECLARE_SCOPE_CYCLE_COUNTER(TEXT("FModelCPU_Init_CreateORTSession"), STAT_FModelCPU_Init_CreateORTSession, STATGROUP_NNE);
+				DECLARE_SCOPE_CYCLE_COUNTER(TEXT("FModelInstanceCPU_Init_CreateORTSession"), STAT_FModelInstanceCPU_Init_CreateORTSession, STATGROUP_NNE);
 				// Read model from InferenceModel
 				Session = MakeUnique<Ort::Session>(*ORTEnvironment, ModelBuffer.GetData(), ModelBuffer.Num(), *SessionOptions);
 			}
@@ -87,12 +87,12 @@ namespace UE::NNERuntimeORTCpu::Private
 		return IsLoaded();
 	}
 
-	bool FModelCPU::IsLoaded() const
+	bool FModelInstanceCPU::IsLoaded() const
 	{
 		return bIsLoaded;
 	}
 
-	bool FModelCPU::InitializedAndConfigureMembers()
+	bool FModelInstanceCPU::InitializedAndConfigureMembers()
 	{
 		// Initialize 
 		// Set up ORT and create an environment
@@ -115,7 +115,7 @@ namespace UE::NNERuntimeORTCpu::Private
 	}
 
 
-	bool FModelCPU::ConfigureTensors(const bool InIsInput)
+	bool FModelInstanceCPU::ConfigureTensors(const bool InIsInput)
 	{
 		const bool bIsInput = InIsInput;
 
@@ -162,14 +162,14 @@ namespace UE::NNERuntimeORTCpu::Private
 		return true;
 	}
 	
-	int FModelCPU::SetInputTensorShapes(TConstArrayView<NNECore::FTensorShape> InInputShapes)
+	int FModelInstanceCPU::SetInputTensorShapes(TConstArrayView<NNECore::FTensorShape> InInputShapes)
 	{
 		InputTensors.Reset();
 		OutputTensors.Reset();
 		OutputTensorShapes.Reset();
 
 		// Verify input shape are valid for the model and set InputTensorShapes
-		if (FModelBase<IModelCPU>::SetInputTensorShapes(InInputShapes) != 0)
+		if (FModelInstanceBase<IModelInstanceCPU>::SetInputTensorShapes(InInputShapes) != 0)
 		{
 			return -1;
 		}
@@ -203,14 +203,14 @@ namespace UE::NNERuntimeORTCpu::Private
 		return 0;
 	}
 
-	int FModelCPU::RunSync(TConstArrayView<NNECore::FTensorBindingCPU> InInputBindings, TConstArrayView<NNECore::FTensorBindingCPU> InOutputBindings)
+	int FModelInstanceCPU::RunSync(TConstArrayView<NNECore::FTensorBindingCPU> InInputBindings, TConstArrayView<NNECore::FTensorBindingCPU> InOutputBindings)
 	{
-		DECLARE_SCOPE_CYCLE_COUNTER(TEXT("FModelCPU_Run"), STAT_FModelCPU_Run, STATGROUP_NNE);
+		DECLARE_SCOPE_CYCLE_COUNTER(TEXT("FModelInstanceCPU_Run"), STAT_FModelInstanceCPU_Run, STATGROUP_NNE);
 
 		// Sanity check
 		if (!bIsLoaded)
 		{
-			UE_LOG(LogNNE, Warning, TEXT("FModelCPU::Run(): Call FModelCPU::Load() to load a model first."));
+			UE_LOG(LogNNE, Warning, TEXT("FModelInstanceCPU::Run(): Call FModelInstanceCPU::Load() to load a model first."));
 			return -1;
 		}
 
@@ -279,25 +279,47 @@ namespace UE::NNERuntimeORTCpu::Private
 		return 0;
 	}
 
-	float FModelCPU::GetLastRunTimeMSec() const
+	float FModelInstanceCPU::GetLastRunTimeMSec() const
 	{
 		return RunStatisticsEstimator.GetLastSample();
 	}
 
-	NNEProfiling::Internal::FStatistics FModelCPU::GetRunStatistics() const
+	NNEProfiling::Internal::FStatistics FModelInstanceCPU::GetRunStatistics() const
 	{
 		return RunStatisticsEstimator.GetStats();
 	}
 
-	NNEProfiling::Internal::FStatistics FModelCPU::GetInputMemoryTransferStats() const
+	NNEProfiling::Internal::FStatistics FModelInstanceCPU::GetInputMemoryTransferStats() const
 	{
 		return InputTransferStatisticsEstimator.GetStats();
 	}
 
-	void FModelCPU::ResetStats()
+	void FModelInstanceCPU::ResetStats()
 	{
 		RunStatisticsEstimator.ResetStats();
 		InputTransferStatisticsEstimator.ResetStats();
 	}
 
+	TUniquePtr<UE::NNECore::IModelInstanceCPU> FModelCPU::CreateModelInstance()
+	{
+		const UE::NNERuntimeORTCpu::Private::FRuntimeConf InConf;
+		UE::NNERuntimeORTCpu::Private::FModelInstanceCPU* ModelInstance = new UE::NNERuntimeORTCpu::Private::FModelInstanceCPU(ORTEnvironment, InConf);
+
+		if (!ModelInstance->Init(ModelData))
+		{
+			delete ModelInstance;
+			return TUniquePtr<UE::NNECore::IModelInstanceCPU>();
+		}
+
+		UE::NNECore::IModelInstanceCPU* IModelInstance = static_cast<UE::NNECore::IModelInstanceCPU*>(ModelInstance);
+		return TUniquePtr<UE::NNECore::IModelInstanceCPU>(IModelInstance);
+	}
+
+	FModelCPU::FModelCPU(Ort::Env* InORTEnvironment, TConstArrayView<uint8> InModelData) :
+		ORTEnvironment(InORTEnvironment),
+		ModelData(InModelData)
+	{
+
+	}
+
 } // namespace UE::NNERuntimeORTCpu::Private
\ No newline at end of file
diff --git a/Engine/Plugins/Experimental/NNE/Source/NNERuntimeORTCpu/Private/NNERuntimeORTCpuModel.h b/Engine/Plugins/Experimental/NNE/Source/NNERuntimeORTCpu/Private/NNERuntimeORTCpuModel.h
--- a/Engine/Plugins/Experimental/NNE/Source/NNERuntimeORTCpu/Private/NNERuntimeORTCpuModel.h
+++ b/Engine/Plugins/Experimental/NNE/Source/NNERuntimeORTCpu/Private/NNERuntimeORTCpuModel.h
@@ -25,13 +25,13 @@ namespace UE::NNERuntimeORTCpu::Private
 		EThreadPriority ThreadPriority = EThreadPriority::TPri_Normal;
 	};
 
-	class FModelCPU : public NNECore::Internal::FModelBase<NNECore::IModelCPU>
+	class FModelInstanceCPU : public NNECore::Internal::FModelInstanceBase<NNECore::IModelInstanceCPU>
 	{
 
 	public:
-		FModelCPU();
-		FModelCPU(Ort::Env* InORTEnvironment, const FRuntimeConf& InRuntimeConf);
-		virtual ~FModelCPU() {};
+		FModelInstanceCPU();
+		FModelInstanceCPU(Ort::Env* InORTEnvironment, const FRuntimeConf& InRuntimeConf);
+		virtual ~FModelInstanceCPU() {};
 
 		virtual int SetInputTensorShapes(TConstArrayView<NNECore::FTensorShape> InInputShapes) override;
 		virtual int RunSync(TConstArrayView<NNECore::FTensorBindingCPU> InInputBindings, TConstArrayView<NNECore::FTensorBindingCPU> InOutputBindings) override;
@@ -79,5 +79,18 @@ namespace UE::NNERuntimeORTCpu::Private
 		UE::NNEProfiling::Internal::FStatisticsEstimator InputTransferStatisticsEstimator;
 
 	};
+
+	class FModelCPU : public NNECore::IModelCPU
+	{
+	public:
+		FModelCPU(Ort::Env* InORTEnvironment, TConstArrayView<uint8> ModelData);
+		virtual ~FModelCPU() {};
+
+		virtual TUniquePtr<UE::NNECore::IModelInstanceCPU> CreateModelInstance() override;
+
+	private:
+		Ort::Env* ORTEnvironment;
+		TArray<uint8> ModelData;
+	};
 	
 } // UE::NNERuntimeORTCpu::Private
\ No newline at end of file
diff --git a/Engine/Plugins/Experimental/NNE/Source/NNERuntimeRDG/Private/Dml/NNEDmlModel.cpp b/Engine/Plugins/Experimental/NNE/Source/NNERuntimeRDG/Private/Dml/NNEDmlModel.cpp
--- a/Engine/Plugins/Experimental/NNE/Source/NNERuntimeRDG/Private/Dml/NNEDmlModel.cpp
+++ b/Engine/Plugins/Experimental/NNE/Source/NNERuntimeRDG/Private/Dml/NNEDmlModel.cpp
@@ -62,7 +62,7 @@ bool FModelInfo::ValidateGuidAndVersion(const uint8* InGuid, const uint8* InVers
 //
 //
 //
-class FModel::FDebugName
+class FModelInstance::FDebugName
 {
 	static constexpr int32 Size = 128;
 
@@ -110,11 +110,11 @@ private:
 //
 //
 //
-class FModel::FBindingTable
+class FModelInstance::FBindingTable
 {
 public:
 
-	bool Init(FModel* InModel)
+	bool Init(FModelInstance* InModel)
 	{
 		Model = InModel;
 		DynamicRHI = InModel->DynamicRHI;
@@ -299,13 +299,13 @@ private:
 	TArray<DML_BUFFER_BINDING, TInlineAllocator<MaxNumOutputs>>		OutputBinds;
 	TArray<DML_BINDING_DESC, TInlineAllocator<MaxNumOutputs>>		OutputBindDescs;
 	ID3D12DynamicRHI*												DynamicRHI;
-	FModel*															Model;
+	FModelInstance*															Model;
 };
 
 //
 //
 //
-class FModel::FGraphBuilder
+class FModelInstance::FGraphBuilder
 {
 public:
 
@@ -822,21 +822,21 @@ private:
 //
 //
 //
-FModel::FModel()
+FModelInstance::FModelInstance()
 {
 }
 
 //
 //
 //
-FModel::~FModel()
+FModelInstance::~FModelInstance()
 {
 }
 
 //
 //
 //
-bool FModel::Init(TConstArrayView<uint8> ModelData, FDmlDeviceContext* InDevCtx)
+bool FModelInstance::Init(TConstArrayView<uint8> ModelData, FDmlDeviceContext* InDevCtx)
 {
 	ConstantCPUTensorIndices.Reset();
 
@@ -984,7 +984,7 @@ bool FModel::Init(TConstArrayView<uint8> ModelData, FDmlDeviceContext* InDevCtx)
 //
 //
 //
-bool FModel::InitCompiledOp(TConstArrayView<int32> OpInputIndices, uint64 TensorDataSize)
+bool FModelInstance::InitCompiledOp(TConstArrayView<int32> OpInputIndices, uint64 TensorDataSize)
 {
 	static constexpr EBufferUsageFlags	WeightBuffUsage = BUF_UnorderedAccess;
 	static constexpr ERHIAccess			WeightBuffAccess = ERHIAccess::UAVMask;
@@ -1178,7 +1178,7 @@ END_SHADER_PARAMETER_STRUCT()
 //
 //
 //
-void FModel::AddDispatchOps_RenderThread(FRDGBuilder& GraphBuilder)
+void FModelInstance::AddDispatchOps_RenderThread(FRDGBuilder& GraphBuilder)
 {
 	FDmlModelDispatchPassParameters* DispatchParams = GraphBuilder.AllocParameters<FDmlModelDispatchPassParameters>();
 
@@ -1299,7 +1299,7 @@ void FModel::AddDispatchOps_RenderThread(FRDGBuilder& GraphBuilder)
 //
 // Create operator
 //
-FOperatorDml* FModel::OpCreate(const FString& OpName, TArrayView<const NNECore::Internal::FTensor> InputTensorDescs, TArrayView<const NNECore::Internal::FTensor> OutputTensorDescs, const NNECore::FAttributeMap& Attributes)
+FOperatorDml* FModelInstance::OpCreate(const FString& OpName, TArrayView<const NNECore::Internal::FTensor> InputTensorDescs, TArrayView<const NNECore::Internal::FTensor> OutputTensorDescs, const NNECore::FAttributeMap& Attributes)
 {
 	FOperatorRegistryDml::OperatorCreateFunc CreateFn = FOperatorRegistryDml::Get()->OpFind(OpName);
 
@@ -1324,7 +1324,7 @@ FOperatorDml* FModel::OpCreate(const FString& OpName, TArrayView<const NNECore::
 	return Op;
 }
 
-FBufferRHIRef FModel::CreateRHIBuffer(FRHICommandListImmediate& RHICmdList, uint32 Size, EBufferUsageFlags Usage, ERHIAccess Access, const TCHAR* DbgName)
+FBufferRHIRef FModelInstance::CreateRHIBuffer(FRHICommandListImmediate& RHICmdList, uint32 Size, EBufferUsageFlags Usage, ERHIAccess Access, const TCHAR* DbgName)
 {
 	FBufferRHIRef Buff = nullptr;
 
@@ -1339,7 +1339,7 @@ FBufferRHIRef FModel::CreateRHIBuffer(FRHICommandListImmediate& RHICmdList, uint
 	return Buff;
 }
 
-ID3D12Resource* FModel::CreateD3D12Buffer(uint32 Size, D3D12_RESOURCE_STATES ResourceState, D3D12_HEAP_TYPE HeapType, const TCHAR* DebugName)
+ID3D12Resource* FModelInstance::CreateD3D12Buffer(uint32 Size, D3D12_RESOURCE_STATES ResourceState, D3D12_HEAP_TYPE HeapType, const TCHAR* DebugName)
 {
 	ID3D12Resource* Resource = nullptr;
 
@@ -1372,7 +1372,7 @@ ID3D12Resource* FModel::CreateD3D12Buffer(uint32 Size, D3D12_RESOURCE_STATES Res
 //
 //
 //
-int FModel::PrepareTensorShapesAndData()
+int FModelInstance::PrepareTensorShapesAndData()
 {
 	for (NNECore::FTensorDesc SymbolicTensorDesc : AllSymbolicTensorDescs)
 	{
@@ -1386,6 +1386,25 @@ int FModel::PrepareTensorShapesAndData()
 	return 0;
 }
 
+TUniquePtr<UE::NNECore::IModelInstanceRDG> FModel::CreateModelInstance()
+{
+	FModelInstance* ModelInstance = new FModelInstance();
+
+	if (!ModelInstance->Init(ModelData, DevCtx))
+	{
+		delete ModelInstance;
+		return TUniquePtr<UE::NNECore::IModelInstanceRDG>();
+	}
+
+	UE::NNECore::IModelInstanceRDG* IModelInstance = static_cast<UE::NNECore::IModelInstanceRDG*>(ModelInstance);
+	return TUniquePtr<UE::NNECore::IModelInstanceRDG>(IModelInstance);
+}
+
+FModel::FModel(TConstArrayView<uint8> InModelData, FDmlDeviceContext* InDevCtx) 
+	: ModelData(InModelData), DevCtx(InDevCtx)
+{
+}
+
 #endif // NNE_USE_DIRECTML
 
 } // namespace UE::NNERuntimeRDG::Private::Dml
diff --git a/Engine/Plugins/Experimental/NNE/Source/NNERuntimeRDG/Private/Dml/NNEDmlModel.h b/Engine/Plugins/Experimental/NNE/Source/NNERuntimeRDG/Private/Dml/NNEDmlModel.h
--- a/Engine/Plugins/Experimental/NNE/Source/NNERuntimeRDG/Private/Dml/NNEDmlModel.h
+++ b/Engine/Plugins/Experimental/NNE/Source/NNERuntimeRDG/Private/Dml/NNEDmlModel.h
@@ -42,7 +42,7 @@ class FOperatorDml;
 //
 //
 //
-class FModel : public FModelRDG
+class FModelInstance : public FModelInstanceRDG
 {
 	class FGraphBuilder;
 	class FBindingTable;
@@ -50,8 +50,8 @@ class FModel : public FModelRDG
 
 public:
 
-	FModel();
-	~FModel();
+	FModelInstance();
+	~FModelInstance();
 
 	bool Init(TConstArrayView<uint8> ModelData, FDmlDeviceContext* InDevCtx);
 
@@ -97,6 +97,19 @@ private:
 	ID3D12DynamicRHI*					DynamicRHI;
 };
 
+class FModel : public NNECore::IModelRDG
+{
+public:
+	FModel(TConstArrayView<uint8> InModelData, FDmlDeviceContext* InDevCtx);
+	virtual ~FModel() {};
+
+	virtual TUniquePtr<UE::NNECore::IModelInstanceRDG> CreateModelInstance() override;
+
+private:
+	TArray<uint8> ModelData;
+	FDmlDeviceContext* DevCtx;
+};
+
 #endif // NNE_USE_DIRECTML
 
 } // namespace UE::NNERuntimeRDG::Private::Dml
diff --git a/Engine/Plugins/Experimental/NNE/Source/NNERuntimeRDG/Private/NNERuntimeRDGDml.cpp b/Engine/Plugins/Experimental/NNE/Source/NNERuntimeRDG/Private/NNERuntimeRDGDml.cpp
--- a/Engine/Plugins/Experimental/NNE/Source/NNERuntimeRDG/Private/NNERuntimeRDGDml.cpp
+++ b/Engine/Plugins/Experimental/NNE/Source/NNERuntimeRDG/Private/NNERuntimeRDGDml.cpp
@@ -147,26 +147,18 @@ bool UNNERuntimeRDGDmlImpl::CanCreateModelRDG(TObjectPtr<UNNEModelData> ModelDat
 #endif
 };
 
-TUniquePtr<UE::NNECore::IModelRDG> UNNERuntimeRDGDmlImpl::CreateModelRDG(TObjectPtr<UNNEModelData> ModelData)
+TUniquePtr<UE::NNECore::IModelRDG> UNNERuntimeRDGDmlImpl::CreateModel(TObjectPtr<UNNEModelData> ModelData)
 {
 #ifdef NNE_USE_DIRECTML
-	TUniquePtr<FModel>	Model;
-
 	if (!CanCreateModelRDG(ModelData))
 	{
-		return Model;
+		return TUniquePtr<UE::NNECore::IModelRDG>();
 	}
 
-	// Create the model and initialize it with the data not including the header
-	Model = MakeUnique<FModel>();
 	TConstArrayView<uint8> Data = ModelData->GetModelData(GetRuntimeName());
+	FModel* Model = new FModel(Data, Ctx);
 
-	if (!Model->Init(Data, Ctx))
-	{
-		Model.Reset();
-	}
-
-	return Model;
+	return TUniquePtr<UE::NNECore::IModelRDG>(Model);
 #else
 	return TUniquePtr<UE::NNECore::IModelRDG>();
 #endif
diff --git a/Engine/Plugins/Experimental/NNE/Source/NNERuntimeRDG/Private/NNERuntimeRDGDml.h b/Engine/Plugins/Experimental/NNE/Source/NNERuntimeRDG/Private/NNERuntimeRDGDml.h
--- a/Engine/Plugins/Experimental/NNE/Source/NNERuntimeRDG/Private/NNERuntimeRDGDml.h
+++ b/Engine/Plugins/Experimental/NNE/Source/NNERuntimeRDG/Private/NNERuntimeRDGDml.h
@@ -38,7 +38,7 @@ public:
 	virtual TArray<uint8> CreateModelData(FString FileType, TConstArrayView<uint8> FileData) override;
 
 	virtual bool CanCreateModelRDG(TObjectPtr<UNNEModelData> ModelData) const override;
-	virtual TUniquePtr<UE::NNECore::IModelRDG> CreateModelRDG(TObjectPtr<UNNEModelData> ModelData) override;
+	virtual TUniquePtr<UE::NNECore::IModelRDG> CreateModel(TObjectPtr<UNNEModelData> ModelData) override;
 private:
 	UE::NNERuntimeRDG::Private::Dml::FDmlDeviceContext* Ctx;
 };
diff --git a/Engine/Plugins/Experimental/NNE/Source/NNERuntimeRDG/Private/NNERuntimeRDGHlsl.cpp b/Engine/Plugins/Experimental/NNE/Source/NNERuntimeRDG/Private/NNERuntimeRDGHlsl.cpp
--- a/Engine/Plugins/Experimental/NNE/Source/NNERuntimeRDG/Private/NNERuntimeRDGHlsl.cpp
+++ b/Engine/Plugins/Experimental/NNE/Source/NNERuntimeRDG/Private/NNERuntimeRDGHlsl.cpp
@@ -118,20 +118,15 @@ TArray<uint8> UNNERuntimeRDGHlslImpl::CreateModelData(FString FileType, TConstAr
 	return Result;
 };
 
-TUniquePtr<UE::NNECore::IModelRDG> UNNERuntimeRDGHlslImpl::CreateModelRDG(TObjectPtr<UNNEModelData> ModelData)
+TUniquePtr<UE::NNECore::IModelRDG> UNNERuntimeRDGHlslImpl::CreateModel(TObjectPtr<UNNEModelData> ModelData)
 {
 	if (!CanCreateModelRDG(ModelData))
 	{
 		return TUniquePtr<UE::NNECore::IModelRDG>();
 	}
 
-	UE::NNERuntimeRDG::Private::Hlsl::FModel* Model = new UE::NNERuntimeRDG::Private::Hlsl::FModel();
 	TConstArrayView<uint8> Data = ModelData->GetModelData(GetRuntimeName());
-	
-	if (!Model->Init(Data))
-	{
-		delete Model;
-		return TUniquePtr<UE::NNECore::IModelRDG>();
-	}
+	UE::NNERuntimeRDG::Private::Hlsl::FModel* Model = new UE::NNERuntimeRDG::Private::Hlsl::FModel(Data);
+
 	return TUniquePtr<UE::NNECore::IModelRDG>(Model);
 }
diff --git a/Engine/Plugins/Experimental/NNE/Source/NNERuntimeRDG/Private/NNERuntimeRDGHlsl.h b/Engine/Plugins/Experimental/NNE/Source/NNERuntimeRDG/Private/NNERuntimeRDGHlsl.h
--- a/Engine/Plugins/Experimental/NNE/Source/NNERuntimeRDG/Private/NNERuntimeRDGHlsl.h
+++ b/Engine/Plugins/Experimental/NNE/Source/NNERuntimeRDG/Private/NNERuntimeRDGHlsl.h
@@ -33,5 +33,5 @@ public:
 	virtual TArray<uint8> CreateModelData(FString FileType, TConstArrayView<uint8> FileData) override;
 
 	virtual bool CanCreateModelRDG(TObjectPtr<UNNEModelData> ModelData) const override;
-	virtual TUniquePtr<UE::NNECore::IModelRDG> CreateModelRDG(TObjectPtr<UNNEModelData> ModelData) override;
+	virtual TUniquePtr<UE::NNECore::IModelRDG> CreateModel(TObjectPtr<UNNEModelData> ModelData) override;
 };
\ No newline at end of file
diff --git a/Engine/Plugins/Experimental/NNE/Source/NNERuntimeRDG/Private/NNERuntimeRDGModel.cpp b/Engine/Plugins/Experimental/NNE/Source/NNERuntimeRDG/Private/NNERuntimeRDGModel.cpp
--- a/Engine/Plugins/Experimental/NNE/Source/NNERuntimeRDG/Private/NNERuntimeRDGModel.cpp
+++ b/Engine/Plugins/Experimental/NNE/Source/NNERuntimeRDG/Private/NNERuntimeRDGModel.cpp
@@ -13,7 +13,7 @@
 namespace UE::NNERuntimeRDG::Private
 {
 
-bool FModelRDG::LoadModel(TConstArrayView<uint8> ModelData, FNNERuntimeFormat& Format, int32 GuidAndVersionSize)
+bool FModelInstanceRDG::LoadModel(TConstArrayView<uint8> ModelData, FNNERuntimeFormat& Format, int32 GuidAndVersionSize)
 {
 	TConstArrayView<uint8> ModelBuffer = { &(ModelData.GetData()[GuidAndVersionSize]), ModelData.Num() - GuidAndVersionSize };
 
@@ -96,12 +96,12 @@ bool FModelRDG::LoadModel(TConstArrayView<uint8> ModelData, FNNERuntimeFormat& F
 
 
 
-int FModelRDG::SetInputTensorShapes(TConstArrayView<NNECore::FTensorShape> InInputShapes)
+int FModelInstanceRDG::SetInputTensorShapes(TConstArrayView<NNECore::FTensorShape> InInputShapes)
 {
 	OutputTensorShapes.Reset(OutputTensorIndices.Num());
 
 	//Verify input shape are valid for the model and set InputTensorShapes
-	if (FModelBase<NNECore::IModelRDG>::SetInputTensorShapes(InInputShapes) != 0)
+	if (FModelInstanceBase<NNECore::IModelInstanceRDG>::SetInputTensorShapes(InInputShapes) != 0)
 	{
 		return -1;
 	}
@@ -197,7 +197,7 @@ FRDGBufferDesc CreateRDGBufferDescForTensorRDG(const FTensorRDG& Tensor)
 /**
  * Enqueue operators to RDG, the caller will run the GraphBuilder.Execute()
  */
-int FModelRDG::EnqueueRDG(FRDGBuilder& RDGBuilder, TConstArrayView<NNECore::FTensorBindingRDG> InInputBindings, TConstArrayView<NNECore::FTensorBindingRDG> InOutputBindings)
+int FModelInstanceRDG::EnqueueRDG(FRDGBuilder& RDGBuilder, TConstArrayView<NNECore::FTensorBindingRDG> InInputBindings, TConstArrayView<NNECore::FTensorBindingRDG> InOutputBindings)
 {
 	check(IsInRenderingThread());
 
@@ -253,7 +253,7 @@ int FModelRDG::EnqueueRDG(FRDGBuilder& RDGBuilder, TConstArrayView<NNECore::FTen
 	return 0;
 }
 
-int FModelRDG::SetTensors(FRDGBuilder& GraphBuilder, FTensorRDGArray& InTensorRDGs, TConstArrayView<NNECore::FTensorBindingRDG> InBindings)
+int FModelInstanceRDG::SetTensors(FRDGBuilder& GraphBuilder, FTensorRDGArray& InTensorRDGs, TConstArrayView<NNECore::FTensorBindingRDG> InBindings)
 {
 	check(InBindings.Num() == InTensorRDGs.Num());
 	
diff --git a/Engine/Plugins/Experimental/NNE/Source/NNERuntimeRDG/Private/NNERuntimeRDGModel.h b/Engine/Plugins/Experimental/NNE/Source/NNERuntimeRDG/Private/NNERuntimeRDGModel.h
--- a/Engine/Plugins/Experimental/NNE/Source/NNERuntimeRDG/Private/NNERuntimeRDGModel.h
+++ b/Engine/Plugins/Experimental/NNE/Source/NNERuntimeRDG/Private/NNERuntimeRDGModel.h
@@ -11,11 +11,11 @@
 
 namespace UE::NNERuntimeRDG::Private
 {
-	class FModelRDG : public NNECore::Internal::FModelBase<NNECore::IModelRDG>
+	class FModelInstanceRDG : public NNECore::Internal::FModelInstanceBase<NNECore::IModelInstanceRDG>
 	{
 	public:
-		FModelRDG() {};
-		virtual ~FModelRDG() = default;
+		FModelInstanceRDG() {};
+		virtual ~FModelInstanceRDG() = default;
 
 		virtual int SetInputTensorShapes(TConstArrayView<NNECore::FTensorShape> InInputShapes) override;
 		virtual int EnqueueRDG(FRDGBuilder& RDGBuilder, TConstArrayView<NNECore::FTensorBindingRDG> InInputBindings, TConstArrayView<NNECore::FTensorBindingRDG> InOutputBindings) override;
diff --git a/Engine/Plugins/Experimental/NNE/Source/NNERuntimeRDG/Private/NNERuntimeRDGModelHlsl.cpp b/Engine/Plugins/Experimental/NNE/Source/NNERuntimeRDG/Private/NNERuntimeRDGModelHlsl.cpp
--- a/Engine/Plugins/Experimental/NNE/Source/NNERuntimeRDG/Private/NNERuntimeRDGModelHlsl.cpp
+++ b/Engine/Plugins/Experimental/NNE/Source/NNERuntimeRDG/Private/NNERuntimeRDGModelHlsl.cpp
@@ -60,7 +60,7 @@ FTensorRDGRef GetTensorRefIfWeightAndOnlyUsedByOperator(int32 TensorIndex, int32
 
 } // namespace ModelUtils
 
-bool FModel::PrepareModelRDG(FRDGBuilder& RDGBuilder)
+bool FModelInstance::PrepareModelRDG(FRDGBuilder& RDGBuilder)
 {
 	//Register constant tensors to graph, uploading if needed
 	check(IntermediateTensorRDGs.Num() == ConstantsExternalRDGResources.Num());
@@ -104,7 +104,7 @@ bool FModel::PrepareModelRDG(FRDGBuilder& RDGBuilder)
 	return true;
 }
 
-bool FModel::Init(TConstArrayView<uint8> ModelData)
+bool FModelInstance::Init(TConstArrayView<uint8> ModelData)
 {
 	check(ModelData.Num() > 0);
 	FNNERuntimeFormat	Format;
@@ -166,7 +166,7 @@ bool FModel::Init(TConstArrayView<uint8> ModelData)
 	return true;
 }
 
-void FModel::AddDispatchOps_RenderThread(FRDGBuilder& GraphBuilder)
+void FModelInstance::AddDispatchOps_RenderThread(FRDGBuilder& GraphBuilder)
 {
 	static constexpr int32 MaxExpectedInput = 10;
 	TArray<FTensorRDGRef, TInlineAllocator<MaxExpectedInput>> InputTensors;
@@ -207,7 +207,7 @@ void FModel::AddDispatchOps_RenderThread(FRDGBuilder& GraphBuilder)
 	}
 }
 
-int FModel::PrepareTensorShapesAndData()
+int FModelInstance::PrepareTensorShapesAndData()
 {
 	check(AllTensorRDGRefs.Num() == AllSymbolicTensorDescs.Num());
 	
@@ -307,7 +307,7 @@ void EnqueueTensorUpload(TArray<TRefCountPtr<FRDGPooledBuffer>>& OutExternalRDGR
 
 		FEvent* Signal = FGenericPlatformProcess::GetSynchEventFromPool(false);
 
-		ENQUEUE_RENDER_COMMAND(FModel_UploadTensors)
+		ENQUEUE_RENDER_COMMAND(FModelInstance_UploadTensors)
 		(
 			[Signal, &OutExternalRDGResources, &TensorToUploadRDGs, CopyDataFlag](FRHICommandListImmediate& RHICmdList)
 			{
@@ -354,7 +354,7 @@ void EnqueueTensorUpload(TArray<TRefCountPtr<FRDGPooledBuffer>>& OutExternalRDGR
 	}
 }
 
-bool FModel::PrepareWeights()
+bool FModelInstance::PrepareWeights()
 {
 	check(WeightsExternalRDGResources.IsEmpty());
 
@@ -364,4 +364,20 @@ bool FModel::PrepareWeights()
 	return true;
 }
 
+TUniquePtr<UE::NNECore::IModelInstanceRDG> FModel::CreateModelInstance()
+{
+	FModelInstance* ModelInstance = new FModelInstance();
+
+	if (!ModelInstance->Init(ModelData))
+	{
+		delete ModelInstance;
+		return TUniquePtr<UE::NNECore::IModelInstanceRDG>();
+	}
+
+	UE::NNECore::IModelInstanceRDG* IModelInstance = static_cast<UE::NNECore::IModelInstanceRDG*>(ModelInstance);
+	return TUniquePtr<UE::NNECore::IModelInstanceRDG>(IModelInstance);
+}
+
+FModel::FModel(TConstArrayView<uint8> InModelData) : ModelData(InModelData) {}
+
 } // namespace UE::NNERuntimeRDG::Private::Hlsl
\ No newline at end of file
diff --git a/Engine/Plugins/Experimental/NNE/Source/NNERuntimeRDG/Private/NNERuntimeRDGModelHlsl.h b/Engine/Plugins/Experimental/NNE/Source/NNERuntimeRDG/Private/NNERuntimeRDGModelHlsl.h
--- a/Engine/Plugins/Experimental/NNE/Source/NNERuntimeRDG/Private/NNERuntimeRDGModelHlsl.h
+++ b/Engine/Plugins/Experimental/NNE/Source/NNERuntimeRDG/Private/NNERuntimeRDGModelHlsl.h
@@ -11,12 +11,12 @@ namespace UE::NNERuntimeRDG::Private::Hlsl
 {
 struct FOperatorHlsl;
 
-class FModel : public FModelRDG
+class FModelInstance : public FModelInstanceRDG
 {
 	
 public:
 
-	~FModel() = default;
+	~FModelInstance() = default;
 
 	bool Init(TConstArrayView<uint8> ModelData);
 
@@ -35,4 +35,16 @@ private:
 	TArray<TRefCountPtr<FRDGPooledBuffer>> ConstantsExternalRDGResources;
 };
 
+class FModel : public NNECore::IModelRDG
+{
+public:
+	FModel(TConstArrayView<uint8> ModelData);
+	virtual ~FModel() {};
+
+	virtual TUniquePtr<UE::NNECore::IModelInstanceRDG> CreateModelInstance() override;
+
+private:
+	TArray<uint8> ModelData;
+};
+
 } // namespace UE::NNERuntimeRDG::Private::Hlsl
\ No newline at end of file
diff --git a/Engine/Plugins/Experimental/NNERuntimeORTGpu/Source/NNERuntimeORT/Private/NNERuntimeORT.cpp b/Engine/Plugins/Experimental/NNERuntimeORTGpu/Source/NNERuntimeORT/Private/NNERuntimeORT.cpp
--- a/Engine/Plugins/Experimental/NNERuntimeORTGpu/Source/NNERuntimeORT/Private/NNERuntimeORT.cpp
+++ b/Engine/Plugins/Experimental/NNERuntimeORTGpu/Source/NNERuntimeORT/Private/NNERuntimeORT.cpp
@@ -85,7 +85,7 @@ bool UNNERuntimeORTGpuImpl::CanCreateModelGPU(TObjectPtr<UNNEModelData> ModelDat
 	return bResult;
 }
 
-TUniquePtr<UE::NNECore::IModelGPU> UNNERuntimeORTGpuImpl::CreateModelGPU(TObjectPtr<UNNEModelData> ModelData)
+TUniquePtr<UE::NNECore::IModelGPU> UNNERuntimeORTGpuImpl::CreateModel(TObjectPtr<UNNEModelData> ModelData)
 {
 	check(ModelData != nullptr);
 	check(ORTEnvironment.IsValid());
@@ -95,29 +95,22 @@ TUniquePtr<UE::NNECore::IModelGPU> UNNERuntimeORTGpuImpl::CreateModelGPU(TObject
 		return TUniquePtr<UE::NNECore::IModelGPU>();
 	}
 
-	const UE::NNERuntimeORT::Private::FRuntimeConf InConf;
-	UE::NNERuntimeORT::Private::FModelORT* Model = nullptr;
+	UE::NNECore::IModelGPU* IModel = nullptr;
 	TConstArrayView<uint8> Data = ModelData->GetModelData(GetRuntimeName());
 
 	switch (Provider)
 	{
 		case ENNERuntimeORTGpuProvider::Dml:  
-			Model = new UE::NNERuntimeORT::Private::FModelORTDml(ORTEnvironment.Get(), InConf);
+			IModel = static_cast<UE::NNECore::IModelGPU*>(new UE::NNERuntimeORT::Private::FModelORTDml(ORTEnvironment.Get(), Data));
 			break;
 		case ENNERuntimeORTGpuProvider::Cuda: 
-			Model = new UE::NNERuntimeORT::Private::FModelORTCuda(ORTEnvironment.Get(), InConf);
+			IModel = static_cast<UE::NNECore::IModelGPU*>(new UE::NNERuntimeORT::Private::FModelORTCuda(ORTEnvironment.Get(), Data));
 			break;
 		default:
 			UE_LOG(LogNNE, Error, TEXT("Failed to create model for ORT GPU runtime, unsupported provider. Runtime will not be functional."));
 			return TUniquePtr<UE::NNECore::IModelGPU>();
 	}
 
-	if (!Model->Init(Data))
-	{
-		delete Model;
-		return TUniquePtr<UE::NNECore::IModelGPU>();
-	}
-	UE::NNECore::IModelGPU* IModel = static_cast<UE::NNECore::IModelGPU*>(Model);
 	return TUniquePtr<UE::NNECore::IModelGPU>(IModel);
 }
 
diff --git a/Engine/Plugins/Experimental/NNERuntimeORTGpu/Source/NNERuntimeORT/Private/NNERuntimeORT.h b/Engine/Plugins/Experimental/NNERuntimeORTGpu/Source/NNERuntimeORT/Private/NNERuntimeORT.h
--- a/Engine/Plugins/Experimental/NNERuntimeORTGpu/Source/NNERuntimeORT/Private/NNERuntimeORT.h
+++ b/Engine/Plugins/Experimental/NNERuntimeORTGpu/Source/NNERuntimeORT/Private/NNERuntimeORT.h
@@ -50,5 +50,5 @@ public:
 	virtual TArray<uint8> CreateModelData(FString FileType, TConstArrayView<uint8> FileData) override;
 
 	virtual bool CanCreateModelGPU(TObjectPtr<UNNEModelData> ModelData) const override;
-	virtual TUniquePtr<UE::NNECore::IModelGPU> CreateModelGPU(TObjectPtr<UNNEModelData> ModelData) override;
+	virtual TUniquePtr<UE::NNECore::IModelGPU> CreateModel(TObjectPtr<UNNEModelData> ModelData) override;
 };
\ No newline at end of file
diff --git a/Engine/Plugins/Experimental/NNERuntimeORTGpu/Source/NNERuntimeORT/Private/NNERuntimeORTModel.cpp b/Engine/Plugins/Experimental/NNERuntimeORTGpu/Source/NNERuntimeORT/Private/NNERuntimeORTModel.cpp
--- a/Engine/Plugins/Experimental/NNERuntimeORTGpu/Source/NNERuntimeORT/Private/NNERuntimeORTModel.cpp
+++ b/Engine/Plugins/Experimental/NNERuntimeORTGpu/Source/NNERuntimeORT/Private/NNERuntimeORTModel.cpp
@@ -38,12 +38,12 @@ NNE_THIRD_PARTY_INCLUDES_END
 namespace UE::NNERuntimeORT::Private
 {
 
-	FModelORT::FModelORT() :
+	FModelInstanceORT::FModelInstanceORT() :
 		bIsLoaded(false),
 		bHasRun(false)
 	{ }
 
-	FModelORT::FModelORT(
+	FModelInstanceORT::FModelInstanceORT(
 		Ort::Env* InORTEnvironment,
 		const FRuntimeConf& InRuntimeConf) :
 		bIsLoaded(false),
@@ -52,7 +52,7 @@ namespace UE::NNERuntimeORT::Private
 		ORTEnvironment(InORTEnvironment)
 	{ }
 
-	bool FModelORT::Init(TConstArrayView<uint8> ModelData)
+	bool FModelInstanceORT::Init(TConstArrayView<uint8> ModelData)
 	{
 		DECLARE_SCOPE_CYCLE_COUNTER(TEXT("FModelORT_Init"), STAT_FModelORT_Init, STATGROUP_NNE);
 		
@@ -118,12 +118,12 @@ namespace UE::NNERuntimeORT::Private
 		return IsLoaded();
 	}
 
-	bool FModelORT::IsLoaded() const
+	bool FModelInstanceORT::IsLoaded() const
 	{
 		return bIsLoaded;
 	}
 
-	bool FModelORT::InitializedAndConfigureMembers()
+	bool FModelInstanceORT::InitializedAndConfigureMembers()
 	{
 		// Initialize 
 		// Setting up ORT
@@ -143,7 +143,7 @@ namespace UE::NNERuntimeORT::Private
 	}
 
 
-	bool FModelORT::ConfigureTensors(const bool InIsInput)
+	bool FModelInstanceORT::ConfigureTensors(const bool InIsInput)
 	{
 		const bool bIsInput = InIsInput;
 
@@ -188,14 +188,14 @@ namespace UE::NNERuntimeORT::Private
 		return true;
 	}
 	
-	int FModelORT::SetInputTensorShapes(TConstArrayView<NNECore::FTensorShape> InInputShapes)
+	int FModelInstanceORT::SetInputTensorShapes(TConstArrayView<NNECore::FTensorShape> InInputShapes)
 	{
 		InputTensors.Reset();
 		OutputTensors.Reset();
 		OutputTensorShapes.Reset();
 
 		// Verify input shape are valid for the model and set InputTensorShapes
-		if (FModelBase<IModelGPU>::SetInputTensorShapes(InInputShapes) != 0)
+		if (FModelInstanceBase<IModelInstanceGPU>::SetInputTensorShapes(InInputShapes) != 0)
 		{
 			return -1;
 		}
@@ -229,7 +229,7 @@ namespace UE::NNERuntimeORT::Private
 		return 0;
 	}
 
-	int FModelORT::RunSync(TConstArrayView<NNECore::FTensorBindingGPU> InInputBindings, TConstArrayView<NNECore::FTensorBindingGPU> InOutputBindings)
+	int FModelInstanceORT::RunSync(TConstArrayView<NNECore::FTensorBindingGPU> InInputBindings, TConstArrayView<NNECore::FTensorBindingGPU> InOutputBindings)
 	{
 		DECLARE_SCOPE_CYCLE_COUNTER(TEXT("FModelORT_Run"), STAT_FModelORT_Run, STATGROUP_NNE);
 
@@ -305,31 +305,76 @@ namespace UE::NNERuntimeORT::Private
 		return 0;
 	}
 
-	float FModelORT::GetLastRunTimeMSec() const
+	float FModelInstanceORT::GetLastRunTimeMSec() const
 	{
 		return RunStatisticsEstimator.GetLastSample();
 	}
 
-	NNEProfiling::Internal::FStatistics FModelORT::GetRunStatistics() const
+	NNEProfiling::Internal::FStatistics FModelInstanceORT::GetRunStatistics() const
 	{
 		return RunStatisticsEstimator.GetStats();
 	}
 
-	NNEProfiling::Internal::FStatistics FModelORT::GetInputMemoryTransferStats() const
+	NNEProfiling::Internal::FStatistics FModelInstanceORT::GetInputMemoryTransferStats() const
 	{
 		return InputTransferStatisticsEstimator.GetStats();
 	}
 
-	void FModelORT::ResetStats()
+	void FModelInstanceORT::ResetStats()
 	{
 		RunStatisticsEstimator.ResetStats();
 		InputTransferStatisticsEstimator.ResetStats();
 	}
 
 #if PLATFORM_WINDOWS
-	bool FModelORTDml::InitializedAndConfigureMembers()
+
+	TUniquePtr<UE::NNECore::IModelInstanceGPU> FModelORTDml::CreateModelInstance()
+	{
+		const FRuntimeConf InConf;
+		FModelInstanceORTDml* ModelInstance = new FModelInstanceORTDml(ORTEnvironment, InConf);
+
+		if (!ModelInstance->Init(ModelData))
+		{
+			delete ModelInstance;
+			return TUniquePtr<UE::NNECore::IModelInstanceGPU>();
+		}
+
+		UE::NNECore::IModelInstanceGPU* IModelInstance = static_cast<UE::NNECore::IModelInstanceGPU*>(ModelInstance);
+		return TUniquePtr<UE::NNECore::IModelInstanceGPU>(IModelInstance);
+	}
+
+	FModelORTDml::FModelORTDml(Ort::Env* InORTEnvironment, TConstArrayView<uint8> InModelData) :
+		ORTEnvironment(InORTEnvironment),
+		ModelData(InModelData)
+	{
+
+	}
+
+	TUniquePtr<UE::NNECore::IModelInstanceGPU> FModelORTCuda::CreateModelInstance()
+	{
+		const FRuntimeConf InConf;
+		FModelInstanceORTCuda* ModelInstance = new FModelInstanceORTCuda(ORTEnvironment, InConf);
+
+		if (!ModelInstance->Init(ModelData))
+		{
+			delete ModelInstance;
+			return TUniquePtr<UE::NNECore::IModelInstanceGPU>();
+		}
+
+		UE::NNECore::IModelInstanceGPU* IModelInstance = static_cast<UE::NNECore::IModelInstanceGPU*>(ModelInstance);
+		return TUniquePtr<UE::NNECore::IModelInstanceGPU>(IModelInstance);
+	}
+
+	FModelORTCuda::FModelORTCuda(Ort::Env* InORTEnvironment, TConstArrayView<uint8> InModelData) :
+		ORTEnvironment(InORTEnvironment),
+		ModelData(InModelData)
+	{
+
+	}
+
+	bool FModelInstanceORTDml::InitializedAndConfigureMembers()
 	{
-		if (!FModelORT::InitializedAndConfigureMembers())
+		if (!FModelInstanceORT::InitializedAndConfigureMembers())
 		{
 			return false;
 		}
@@ -398,9 +443,9 @@ namespace UE::NNERuntimeORT::Private
 		return true;
 	}
 
-	bool FModelORTCuda::InitializedAndConfigureMembers()
+	bool FModelInstanceORTCuda::InitializedAndConfigureMembers()
 	{
-		if (!FModelORT::InitializedAndConfigureMembers())
+		if (!FModelInstanceORT::InitializedAndConfigureMembers())
 		{
 			return false;
 		}
diff --git a/Engine/Plugins/Experimental/NNERuntimeORTGpu/Source/NNERuntimeORT/Private/NNERuntimeORTModel.h b/Engine/Plugins/Experimental/NNERuntimeORTGpu/Source/NNERuntimeORT/Private/NNERuntimeORTModel.h
--- a/Engine/Plugins/Experimental/NNERuntimeORTGpu/Source/NNERuntimeORT/Private/NNERuntimeORTModel.h
+++ b/Engine/Plugins/Experimental/NNERuntimeORTGpu/Source/NNERuntimeORT/Private/NNERuntimeORTModel.h
@@ -25,13 +25,13 @@ namespace UE::NNERuntimeORT::Private
 		EThreadPriority ThreadPriority = EThreadPriority::TPri_Normal;
 	};
 
-	class FModelORT : public NNECore::Internal::FModelBase<NNECore::IModelGPU>
+	class FModelInstanceORT : public NNECore::Internal::FModelInstanceBase<NNECore::IModelInstanceGPU>
 	{
 
 	public:
-		FModelORT();
-		FModelORT(Ort::Env* InORTEnvironment, const FRuntimeConf& InRuntimeConf);
-		virtual ~FModelORT() = default;
+		FModelInstanceORT();
+		FModelInstanceORT(Ort::Env* InORTEnvironment, const FRuntimeConf& InRuntimeConf);
+		virtual ~FModelInstanceORT() = default;
 
 		virtual int SetInputTensorShapes(TConstArrayView<NNECore::FTensorShape> InInputShapes) override;
 		virtual int RunSync(TConstArrayView<NNECore::FTensorBindingGPU> InInputBindings, TConstArrayView<NNECore::FTensorBindingGPU> InOutputBindings) override;
@@ -81,25 +81,51 @@ namespace UE::NNERuntimeORT::Private
 	};
 
 #if PLATFORM_WINDOWS
-	class FModelORTDml : public FModelORT
+	class FModelInstanceORTDml : public FModelInstanceORT
 	{
 	public:
-		FModelORTDml() {}
-		FModelORTDml(Ort::Env* InORTEnvironment, const FRuntimeConf& InRuntimeConf) : FModelORT(InORTEnvironment, InRuntimeConf) {}
-		virtual ~FModelORTDml() = default;
+		FModelInstanceORTDml() {}
+		FModelInstanceORTDml(Ort::Env* InORTEnvironment, const FRuntimeConf& InRuntimeConf) : FModelInstanceORT(InORTEnvironment, InRuntimeConf) {}
+		virtual ~FModelInstanceORTDml() = default;
 	private:
 		virtual bool InitializedAndConfigureMembers() override;
 	};
 
-	class FModelORTCuda : public FModelORT
+	class FModelInstanceORTCuda : public FModelInstanceORT
 	{
 	public:
-		FModelORTCuda() {}
-		FModelORTCuda(Ort::Env* InORTEnvironment, const FRuntimeConf& InRuntimeConf) : FModelORT(InORTEnvironment, InRuntimeConf) {}
-		virtual ~FModelORTCuda() = default;
+		FModelInstanceORTCuda() {}
+		FModelInstanceORTCuda(Ort::Env* InORTEnvironment, const FRuntimeConf& InRuntimeConf) : FModelInstanceORT(InORTEnvironment, InRuntimeConf) {}
+		virtual ~FModelInstanceORTCuda() = default;
 	private:
 		virtual bool InitializedAndConfigureMembers() override;
 	};
+
+	class FModelORTDml : public NNECore::IModelGPU
+	{
+	public:
+		FModelORTDml(Ort::Env* InORTEnvironment, TConstArrayView<uint8> ModelData);
+		virtual ~FModelORTDml() {};
+
+		virtual TUniquePtr<UE::NNECore::IModelInstanceGPU> CreateModelInstance() override;
+
+	private:
+		Ort::Env* ORTEnvironment;
+		TArray<uint8> ModelData;
+	};
+
+	class FModelORTCuda : public NNECore::IModelGPU
+	{
+	public:
+		FModelORTCuda(Ort::Env* InORTEnvironment, TConstArrayView<uint8> ModelData);
+		virtual ~FModelORTCuda() {};
+
+		virtual TUniquePtr<UE::NNECore::IModelInstanceGPU> CreateModelInstance() override;
+
+	private:
+		Ort::Env* ORTEnvironment;
+		TArray<uint8> ModelData;
+	};
 #endif
 	
 } // UE::NNERuntimeORT::Private
\ No newline at end of file
