NNE Cleanup: remove LexFromString(ENNEAttributeDataType) from public API

Tests:
* nne all tests + style transfer

#rnx

diff --git a/Engine/Plugins/Experimental/NNE/Source/NNE/Internal/NNEAttributeDataType.h b/Engine/Plugins/Experimental/NNE/Source/NNE/Internal/NNEAttributeDataType.h
--- a/Engine/Plugins/Experimental/NNE/Source/NNE/Internal/NNEAttributeDataType.h
+++ b/Engine/Plugins/Experimental/NNE/Source/NNE/Internal/NNEAttributeDataType.h
@@ -22,9 +22,4 @@ enum class ENNEAttributeDataType : uint8
 	Int32Array,							//!< TArray of 32-bit signed integers
 	String,								//!< built-in FString
 	StringArray							//!< TArray of built-in FString
-};
-
-/**
- * @return ENNEAttributeDataType from the string passed in
- */
-NNE_API void LexFromString(ENNEAttributeDataType& OutValue, const TCHAR* StringVal);
\ No newline at end of file
+};
\ No newline at end of file
diff --git a/Engine/Plugins/Experimental/NNE/Source/NNE/Private/NNEAttributeDataType.cpp b/Engine/Plugins/Experimental/NNE/Source/NNE/Private/NNEAttributeDataType.cpp
deleted file mode 100644
