NNE Backward compatibility for NNEModelData following rename of NNECore plugin to NNE

Tests: nne smoke test + style transfer demo with asset reimport reverted

#rnx

diff --git a/Engine/Plugins/Experimental/NNE/Config/BaseNNE.ini b/Engine/Plugins/Experimental/NNE/Config/BaseNNE.ini
new file mode 100644
--- /dev/null
+++ b/Engine/Plugins/Experimental/NNE/Config/BaseNNE.ini
@@ -0,0 +1,3 @@
+[CoreRedirects]
+
++ClassRedirects=(OldName="/Script/NNECore.NNEModelData", NewName="/Script/NNE.NNEModelData")
\ No newline at end of file
