NNE update plugin description

Test: plugin description in plugin editor + associated links

#rnx

diff --git a/Engine/Plugins/Experimental/NNE/NNE.uplugin b/Engine/Plugins/Experimental/NNE/NNE.uplugin
--- a/Engine/Plugins/Experimental/NNE/NNE.uplugin
+++ b/Engine/Plugins/Experimental/NNE/NNE.uplugin
@@ -2,14 +2,14 @@
 	"FileVersion": 3,
 	"Version": 0.1,
 	"VersionName": "0.1",
-	"FriendlyName": "NNE",
-	"Description": "Cross-platform framework for running deep learning and neural network inference in Unreal Engine.",
+	"FriendlyName": "Neural Network Engine (NNE)",
+	"Description": "Framework for running deep neural networks and other machine learning inference in Unreal Engine. This replaces the Neural Network Inference (NNI) plugin.",
 	"Category": "Experimental",
 	"CreatedBy": "Epic Games, Inc.",
 	"CreatedByURL": "https://epicgames.com",
-	"DocsURL": "",
+	"DocsURL": "https://dev.epicgames.com/community/learning/courses/e7w/unreal-engine-neural-network-engine-nne",
 	"MarketplaceURL": "",
-	"SupportURL": "",
+	"SupportURL": "https://forums.unrealengine.com/",
 	"EnabledByDefault": false,
 	"CanContainContent": false,
 	"IsBetaVersion": false,
