NNEModelData API extension

This change list
* extends the NNEModelData api to give read access to private fields
* adds FileID to INNERuntime::CreateModelData and  INNERuntime::CanCreateModelData to allow a runtime to react to repeated creations of the same model

This change allows MLIR based runtimes to compile shared libraries, cache them and recreate them if Intermediate or binary folders are deleted/cleaned.
The second point is breaking backwards compatibility but the function is only used by our own runtimes and thus the impac minimal.

#rnx

diff --git a/Engine/Plugins/Experimental/NNE/Source/NNE/Private/NNEModelData.cpp b/Engine/Plugins/Experimental/NNE/Source/NNE/Private/NNEModelData.cpp
--- a/Engine/Plugins/Experimental/NNE/Source/NNE/Private/NNEModelData.cpp
+++ b/Engine/Plugins/Experimental/NNE/Source/NNE/Private/NNEModelData.cpp
@@ -35,17 +35,17 @@ FCustomVersionRegistration NNEModelDataVersion(UNNEModelData::GUID, LatestVersio
 
 #if WITH_EDITOR
 
-inline UE::DerivedData::FCacheKey CreateCacheKey(const FGuid& FileDataId, const FString& RuntimeName)
+inline UE::DerivedData::FCacheKey CreateCacheKey(const FGuid& FileId, const FString& RuntimeName)
 {
-	FString GuidString = FileDataId.ToString(EGuidFormats::Digits);
+	FString GuidString = FileId.ToString(EGuidFormats::Digits);
 	return { UE::DerivedData::FCacheBucket(FWideStringView(*GuidString)), FIoHash::HashBuffer(MakeMemoryView(FTCHARToUTF8(RuntimeName))) };
 }
 
-inline FSharedBuffer GetFromDDC(const FGuid& FileDataId, const FString& RuntimeName)
+inline FSharedBuffer GetFromDDC(const FGuid& FileId, const FString& RuntimeName)
 {
 	UE::DerivedData::FCacheGetValueRequest GetRequest;
-	GetRequest.Name = FString("Get-") + RuntimeName + FString("-") + FileDataId.ToString(EGuidFormats::Digits);
-	GetRequest.Key = CreateCacheKey(FileDataId,  RuntimeName);
+	GetRequest.Name = FString("Get-") + RuntimeName + FString("-") + FileId.ToString(EGuidFormats::Digits);
+	GetRequest.Key = CreateCacheKey(FileId,  RuntimeName);
 	FSharedBuffer RawDerivedData;
 	UE::DerivedData::FRequestOwner BlockingGetOwner(UE::DerivedData::EPriority::Blocking);
 	UE::DerivedData::GetCache().GetValue({ GetRequest }, BlockingGetOwner, [&RawDerivedData](UE::DerivedData::FCacheGetValueResponse&& Response)
@@ -56,11 +56,11 @@ inline FSharedBuffer GetFromDDC(const FGuid& FileDataId, const FString& RuntimeN
 	return RawDerivedData;
 }
 
-inline void PutIntoDDC(const FGuid& FileDataId, const FString& RuntimeName, FSharedBuffer& Data)
+inline void PutIntoDDC(const FGuid& FileId, const FString& RuntimeName, FSharedBuffer& Data)
 {
 	UE::DerivedData::FCachePutValueRequest PutRequest;
-	PutRequest.Name = FString("Put-") + RuntimeName + FString("-") + FileDataId.ToString(EGuidFormats::Digits);
-	PutRequest.Key = CreateCacheKey(FileDataId, RuntimeName);
+	PutRequest.Name = FString("Put-") + RuntimeName + FString("-") + FileId.ToString(EGuidFormats::Digits);
+	PutRequest.Key = CreateCacheKey(FileId, RuntimeName);
 	PutRequest.Value = UE::DerivedData::FValue::Compress(Data);
 	UE::DerivedData::FRequestOwner BlockingPutOwner(UE::DerivedData::EPriority::Blocking);
 	UE::DerivedData::GetCache().PutValue({ PutRequest }, BlockingPutOwner);
@@ -69,12 +69,12 @@ inline void PutIntoDDC(const FGuid& FileDataId, const FString& RuntimeName, FSha
 
 #endif
 
-inline TArray<uint8> CreateRuntimeDataBlob(const FString& RuntimeName, FString FileType, const TArray<uint8>& FileData)
+inline TArray<uint8> CreateRuntimeDataBlob(const FString& RuntimeName, FString FileType, const TArray<uint8>& FileData, FGuid FileId)
 {
 	TWeakInterfacePtr<INNERuntime> NNERuntime = UE::NNE::GetRuntime<INNERuntime>(RuntimeName);
 	if (NNERuntime.IsValid())
 	{
-		return NNERuntime->CreateModelData(FileType, FileData);
+		return NNERuntime->CreateModelData(FileType, FileData, FileId);
 	}
 	else
 	{
@@ -92,10 +92,25 @@ void UNNEModelData::Init(const FString& Type, TConstArrayView<uint8> Buffer)
 {
 	FileType = Type;
 	FileData = Buffer;
-	FPlatformMisc::CreateGuid(FileDataId);
+	FPlatformMisc::CreateGuid(FileId);
 	ModelData.Empty();
 }
 
+FString UNNEModelData::GetFileType()
+{
+	return FileType;
+}
+
+TConstArrayView<uint8> UNNEModelData::GetFileData()
+{
+	return FileData;
+}
+
+FGuid UNNEModelData::GetFileId()
+{
+	return FileId;
+}
+
 TConstArrayView<uint8> UNNEModelData::GetModelData(const FString& RuntimeName)
 {
 #if WITH_EDITORONLY_DATA
@@ -121,7 +136,7 @@ TConstArrayView<uint8> UNNEModelData::GetModelData(const FString& RuntimeName)
 	
 #if WITH_EDITOR
 	// Check if we have a remote cache hit
-	FSharedBuffer RemoteData = GetFromDDC(FileDataId, RuntimeName);
+	FSharedBuffer RemoteData = GetFromDDC(FileId, RuntimeName);
 	if (RemoteData.GetSize() > 0)
 	{
 		ModelData.Add(RuntimeName, TArray<uint8>((uint8*)RemoteData.GetData(), RemoteData.GetSize()));
@@ -132,7 +147,7 @@ TConstArrayView<uint8> UNNEModelData::GetModelData(const FString& RuntimeName)
 #endif //WITH_EDITOR
 
 	// Try to create the model
-	TArray<uint8> CreatedData = CreateRuntimeDataBlob(RuntimeName, FileType, FileData);
+	TArray<uint8> CreatedData = CreateRuntimeDataBlob(RuntimeName, FileType, FileData, FileId);
 	if (CreatedData.Num() < 1)
 	{
 		return {};
@@ -144,7 +159,7 @@ TConstArrayView<uint8> UNNEModelData::GetModelData(const FString& RuntimeName)
 #if WITH_EDITOR
 	// And put it into DDC
 	FSharedBuffer SharedBuffer = MakeSharedBufferFromArray(MoveTemp(CreatedData));
-	PutIntoDDC(FileDataId, RuntimeName, SharedBuffer);
+	PutIntoDDC(FileId, RuntimeName, SharedBuffer);
 #endif //WITH_EDITOR
 	
 	TArray<uint8>* CachedCreatedData = ModelData.Find(RuntimeName);
@@ -176,13 +191,13 @@ void UNNEModelData::Serialize(FArchive& Ar)
 
 		for (const FString& RuntimeName : CookedRuntimeNames)
 		{
-			TArray<uint8> CreatedData = CreateRuntimeDataBlob(RuntimeName, FileType, FileData);
+			TArray<uint8> CreatedData = CreateRuntimeDataBlob(RuntimeName, FileType, FileData, FileId);
 			if (CreatedData.Num() > 0)
 			{
 				ModelData.Add(RuntimeName, CreatedData);
 #if WITH_EDITOR
 				FSharedBuffer SharedBuffer = MakeSharedBufferFromArray(MoveTemp(CreatedData));
-				PutIntoDDC(FileDataId, RuntimeName, SharedBuffer);
+				PutIntoDDC(FileId, RuntimeName, SharedBuffer);
 #endif //WITH_EDITOR
 			}
 		}
@@ -195,7 +210,7 @@ void UNNEModelData::Serialize(FArchive& Ar)
 
 		Ar << FileType;
 		Ar << EmptyData;
-		Ar << FileDataId;
+		Ar << FileId;
 		Ar << NumItems;
 
 		for (int i = 0; i < NumItems; i++)
@@ -224,7 +239,7 @@ void UNNEModelData::Serialize(FArchive& Ar)
 
 		Ar << FileType;
 		Ar << FileData;
-		Ar << FileDataId;
+		Ar << FileId;
 		Ar << NumItems;
 
 		if (Ar.IsLoading())
diff --git a/Engine/Plugins/Experimental/NNE/Source/NNE/Public/NNEModelData.h b/Engine/Plugins/Experimental/NNE/Source/NNE/Public/NNEModelData.h
--- a/Engine/Plugins/Experimental/NNE/Source/NNE/Public/NNEModelData.h
+++ b/Engine/Plugins/Experimental/NNE/Source/NNE/Public/NNEModelData.h
@@ -30,6 +30,33 @@ public:
 	 */
 	void Init(const FString& Type, TConstArrayView<uint8> Buffer);
 
+	/**
+	 * Get the type of data inside FileData.
+	 *
+	 * The FileType identifies the type of data inside FileData and typically is the extension of the file used to create the asset.
+	 *
+	 * @return The FileType.
+	 */
+	FString GetFileType();
+
+	/**
+	 * Get read only access to FileData.
+	 *
+	 * The FileData contains the binary data of the file which has been used to create the asset.
+	 *
+	 * @return The FileData.
+	 */
+	TConstArrayView<uint8> GetFileData();
+
+	/**
+	 * Get the FGuid identifying the FileData.
+	 *
+	 * The FileId is created on import of an asset. It can be used to identify the FileData, e.g. when putting corresponding data into the DDC or caching data locally.
+	 *
+	 * @return The FileId.
+	 */
+	FGuid GetFileId();
+
 	/**
 	 * Get the cached (editor) or cooked (game) optimized model data for a given runtime.
 	 *
@@ -99,7 +126,7 @@ private:
 	/**
 	 * A Guid that uniquely identifies this model. This is used to cache optimized models in the editor.
 	 */
-	FGuid FileDataId;
+	FGuid FileId;
 
 	/**
 	 * The processed / optimized model data for the different runtimes.
diff --git a/Engine/Plugins/Experimental/NNE/Source/NNE/Public/NNERuntime.h b/Engine/Plugins/Experimental/NNE/Source/NNE/Public/NNERuntime.h
--- a/Engine/Plugins/Experimental/NNE/Source/NNE/Public/NNERuntime.h
+++ b/Engine/Plugins/Experimental/NNE/Source/NNE/Public/NNERuntime.h
@@ -48,16 +48,18 @@ public:
 	 *
 	 * @param FileType The type of file inside FileData. Corresponds to the file extension (e.g. 'onnx').
 	 * @param FileData The raw binary file of a neural network model.
+	 * @param FileId The unique identifier representing FileData.
 	 * @return True if the runtime is able to create model data, false otherwise.
 	 */
-	virtual bool CanCreateModelData(FString FileType, TConstArrayView<uint8> FileData) const = 0;
+	virtual bool CanCreateModelData(FString FileType, TConstArrayView<uint8> FileData, FGuid FileId) const = 0;
 
 	/**
 	 * Create model data given some raw file data.
 	 *
 	 * @param FileType The type of file inside FileData. Corresponds to the file extension (e.g. 'onnx').
 	 * @param FileData The raw binary file of a neural network model.
+	 * @param FileId The unique identifier representing FileData.
 	 * @return Data representing the runtime specific representation of the model to be stored by UNNEModelData on success or an empty array otherwise.
 	 */
-	virtual TArray<uint8> CreateModelData(FString FileType, TConstArrayView<uint8> FileData) = 0;
+	virtual TArray<uint8> CreateModelData(FString FileType, TConstArrayView<uint8> FileData, FGuid FileId) = 0;
 };
\ No newline at end of file
diff --git a/Engine/Plugins/Experimental/NNERuntimeORTCpu/Source/NNERuntimeORTCpu/Private/NNERuntimeORTCpu.cpp b/Engine/Plugins/Experimental/NNERuntimeORTCpu/Source/NNERuntimeORTCpu/Private/NNERuntimeORTCpu.cpp
--- a/Engine/Plugins/Experimental/NNERuntimeORTCpu/Source/NNERuntimeORTCpu/Private/NNERuntimeORTCpu.cpp
+++ b/Engine/Plugins/Experimental/NNERuntimeORTCpu/Source/NNERuntimeORTCpu/Private/NNERuntimeORTCpu.cpp
@@ -16,14 +16,14 @@
 FGuid UNNERuntimeORTCpuImpl::GUID = FGuid((int32)'O', (int32)'C', (int32)'P', (int32)'U');
 int32 UNNERuntimeORTCpuImpl::Version = 0x00000001;
 
-bool UNNERuntimeORTCpuImpl::CanCreateModelData(FString FileType, TConstArrayView<uint8> FileData) const
+bool UNNERuntimeORTCpuImpl::CanCreateModelData(FString FileType, TConstArrayView<uint8> FileData, FGuid FileId) const
 {
 	return FileType.Compare("onnx", ESearchCase::IgnoreCase) == 0;
 }
 
-TArray<uint8> UNNERuntimeORTCpuImpl::CreateModelData(FString FileType, TConstArrayView<uint8> FileData)
+TArray<uint8> UNNERuntimeORTCpuImpl::CreateModelData(FString FileType, TConstArrayView<uint8> FileData, FGuid FileId)
 {
-	if (!CanCreateModelData(FileType, FileData))
+	if (!CanCreateModelData(FileType, FileData, FileId))
 	{
 		return {};
 	}
diff --git a/Engine/Plugins/Experimental/NNERuntimeORTCpu/Source/NNERuntimeORTCpu/Private/NNERuntimeORTCpu.h b/Engine/Plugins/Experimental/NNERuntimeORTCpu/Source/NNERuntimeORTCpu/Private/NNERuntimeORTCpu.h
--- a/Engine/Plugins/Experimental/NNERuntimeORTCpu/Source/NNERuntimeORTCpu/Private/NNERuntimeORTCpu.h
+++ b/Engine/Plugins/Experimental/NNERuntimeORTCpu/Source/NNERuntimeORTCpu/Private/NNERuntimeORTCpu.h
@@ -34,8 +34,8 @@ public:
 	virtual FString GetRuntimeName() const override { return TEXT("NNERuntimeORTCpu"); };
 	virtual bool IsPlatformSupported(const ITargetPlatform* TargetPlatform) const override { return true; };
 
-	virtual bool CanCreateModelData(FString FileType, TConstArrayView<uint8> FileData) const override;
-	virtual TArray<uint8> CreateModelData(FString FileType, TConstArrayView<uint8> FileData) override;
+	virtual bool CanCreateModelData(FString FileType, TConstArrayView<uint8> FileData, FGuid FileId) const override;
+	virtual TArray<uint8> CreateModelData(FString FileType, TConstArrayView<uint8> FileData, FGuid FileId) override;
 
 	virtual bool CanCreateModelCPU(TObjectPtr<UNNEModelData> ModelData) const override;
 	virtual TUniquePtr<UE::NNE::IModelCPU> CreateModel(TObjectPtr<UNNEModelData> ModelData) override;
diff --git a/Engine/Plugins/Experimental/NNERuntimeORTGpu/Source/NNERuntimeORT/Private/NNERuntimeORT.cpp b/Engine/Plugins/Experimental/NNERuntimeORTGpu/Source/NNERuntimeORT/Private/NNERuntimeORT.cpp
--- a/Engine/Plugins/Experimental/NNERuntimeORTGpu/Source/NNERuntimeORT/Private/NNERuntimeORT.cpp
+++ b/Engine/Plugins/Experimental/NNERuntimeORTGpu/Source/NNERuntimeORT/Private/NNERuntimeORT.cpp
@@ -18,14 +18,14 @@
 FGuid UNNERuntimeORTGpuImpl::GUID = FGuid((int32)'O', (int32)'G', (int32)'P', (int32)'U');
 int32 UNNERuntimeORTGpuImpl::Version = 0x00000001;
 
-bool UNNERuntimeORTGpuImpl::CanCreateModelData(FString FileType, TConstArrayView<uint8> FileData) const
+bool UNNERuntimeORTGpuImpl::CanCreateModelData(FString FileType, TConstArrayView<uint8> FileData, FGuid FileId) const
 {
 	return FileType.Compare("onnx", ESearchCase::IgnoreCase) == 0;
 }
 
-TArray<uint8> UNNERuntimeORTGpuImpl::CreateModelData(FString FileType, TConstArrayView<uint8> FileData)
+TArray<uint8> UNNERuntimeORTGpuImpl::CreateModelData(FString FileType, TConstArrayView<uint8> FileData, FGuid FileId)
 {
-	if (!CanCreateModelData(FileType, FileData))
+	if (!CanCreateModelData(FileType, FileData, FileId))
 	{
 		return {};
 	}
diff --git a/Engine/Plugins/Experimental/NNERuntimeORTGpu/Source/NNERuntimeORT/Private/NNERuntimeORT.h b/Engine/Plugins/Experimental/NNERuntimeORTGpu/Source/NNERuntimeORT/Private/NNERuntimeORT.h
--- a/Engine/Plugins/Experimental/NNERuntimeORTGpu/Source/NNERuntimeORT/Private/NNERuntimeORT.h
+++ b/Engine/Plugins/Experimental/NNERuntimeORTGpu/Source/NNERuntimeORT/Private/NNERuntimeORT.h
@@ -46,8 +46,8 @@ public:
 	virtual FString GetRuntimeName() const override;
 	virtual bool IsPlatformSupported(const ITargetPlatform* TargetPlatform) const override { return true; };
 
-	virtual bool CanCreateModelData(FString FileType, TConstArrayView<uint8> FileData) const override;
-	virtual TArray<uint8> CreateModelData(FString FileType, TConstArrayView<uint8> FileData) override;
+	virtual bool CanCreateModelData(FString FileType, TConstArrayView<uint8> FileData, FGuid FileId) const override;
+	virtual TArray<uint8> CreateModelData(FString FileType, TConstArrayView<uint8> FileData, FGuid FileId) override;
 
 	virtual bool CanCreateModelGPU(TObjectPtr<UNNEModelData> ModelData) const override;
 	virtual TUniquePtr<UE::NNE::IModelGPU> CreateModel(TObjectPtr<UNNEModelData> ModelData) override;
diff --git a/Engine/Plugins/Experimental/NNERuntimeRDG/Source/NNERuntimeRDG/Private/NNERuntimeRDGDml.cpp b/Engine/Plugins/Experimental/NNERuntimeRDG/Source/NNERuntimeRDG/Private/NNERuntimeRDGDml.cpp
--- a/Engine/Plugins/Experimental/NNERuntimeRDG/Source/NNERuntimeRDG/Private/NNERuntimeRDGDml.cpp
+++ b/Engine/Plugins/Experimental/NNERuntimeRDG/Source/NNERuntimeRDG/Private/NNERuntimeRDGDml.cpp
@@ -96,14 +96,14 @@ bool FRuntimeDmlStartup()
 
 using namespace UE::NNERuntimeRDG::Private::Dml;
 
-bool UNNERuntimeRDGDmlImpl::CanCreateModelData(FString FileType, TConstArrayView<uint8> FileData) const
+bool UNNERuntimeRDGDmlImpl::CanCreateModelData(FString FileType, TConstArrayView<uint8> FileData, FGuid FileId) const
 {
 	return FileType.Compare("onnx", ESearchCase::IgnoreCase) == 0;
 }
 
-TArray<uint8> UNNERuntimeRDGDmlImpl::CreateModelData(FString FileType, TConstArrayView<uint8> FileData)
+TArray<uint8> UNNERuntimeRDGDmlImpl::CreateModelData(FString FileType, TConstArrayView<uint8> FileData, FGuid FileId)
 {
-	if (!CanCreateModelData(FileType, FileData))
+	if (!CanCreateModelData(FileType, FileData, FileId))
 	{
 		return {};
 	}
diff --git a/Engine/Plugins/Experimental/NNERuntimeRDG/Source/NNERuntimeRDG/Private/NNERuntimeRDGDml.h b/Engine/Plugins/Experimental/NNERuntimeRDG/Source/NNERuntimeRDG/Private/NNERuntimeRDGDml.h
--- a/Engine/Plugins/Experimental/NNERuntimeRDG/Source/NNERuntimeRDG/Private/NNERuntimeRDGDml.h
+++ b/Engine/Plugins/Experimental/NNERuntimeRDG/Source/NNERuntimeRDG/Private/NNERuntimeRDGDml.h
@@ -34,8 +34,8 @@ public:
 	virtual FString GetRuntimeName() const override { return TEXT("NNERuntimeRDGDml"); };
 	virtual bool IsPlatformSupported(const ITargetPlatform* TargetPlatform) const override { return true; };
 
-	virtual bool CanCreateModelData(FString FileType, TConstArrayView<uint8> FileData) const override;
-	virtual TArray<uint8> CreateModelData(FString FileType, TConstArrayView<uint8> FileData) override;
+	virtual bool CanCreateModelData(FString FileType, TConstArrayView<uint8> FileData, FGuid FileId) const override;
+	virtual TArray<uint8> CreateModelData(FString FileType, TConstArrayView<uint8> FileData, FGuid FileId) override;
 
 	virtual bool CanCreateModelRDG(TObjectPtr<UNNEModelData> ModelData) const override;
 	virtual TUniquePtr<UE::NNE::IModelRDG> CreateModel(TObjectPtr<UNNEModelData> ModelData) override;
diff --git a/Engine/Plugins/Experimental/NNERuntimeRDG/Source/NNERuntimeRDG/Private/NNERuntimeRDGHlsl.cpp b/Engine/Plugins/Experimental/NNERuntimeRDG/Source/NNERuntimeRDG/Private/NNERuntimeRDGHlsl.cpp
--- a/Engine/Plugins/Experimental/NNERuntimeRDG/Source/NNERuntimeRDG/Private/NNERuntimeRDGHlsl.cpp
+++ b/Engine/Plugins/Experimental/NNERuntimeRDG/Source/NNERuntimeRDG/Private/NNERuntimeRDGHlsl.cpp
@@ -73,7 +73,7 @@ bool UNNERuntimeRDGHlslImpl::Init()
 	return true;
 }
 
-bool UNNERuntimeRDGHlslImpl::CanCreateModelData(FString FileType, TConstArrayView<uint8> FileData) const
+bool UNNERuntimeRDGHlslImpl::CanCreateModelData(FString FileType, TConstArrayView<uint8> FileData, FGuid FileId) const
 {
 	return FileType.Compare("onnx", ESearchCase::IgnoreCase) == 0;
 }
@@ -93,9 +93,9 @@ bool UNNERuntimeRDGHlslImpl::CanCreateModelRDG(TObjectPtr<UNNEModelData> ModelDa
 	return bResult;
 };
 
-TArray<uint8> UNNERuntimeRDGHlslImpl::CreateModelData(FString FileType, TConstArrayView<uint8> FileData)
+TArray<uint8> UNNERuntimeRDGHlslImpl::CreateModelData(FString FileType, TConstArrayView<uint8> FileData, FGuid FileId)
 {
-	if (!CanCreateModelData(FileType, FileData))
+	if (!CanCreateModelData(FileType, FileData, FileId))
 	{
 		return {};
 	}
diff --git a/Engine/Plugins/Experimental/NNERuntimeRDG/Source/NNERuntimeRDG/Private/NNERuntimeRDGHlsl.h b/Engine/Plugins/Experimental/NNERuntimeRDG/Source/NNERuntimeRDG/Private/NNERuntimeRDGHlsl.h
--- a/Engine/Plugins/Experimental/NNERuntimeRDG/Source/NNERuntimeRDG/Private/NNERuntimeRDGHlsl.h
+++ b/Engine/Plugins/Experimental/NNERuntimeRDG/Source/NNERuntimeRDG/Private/NNERuntimeRDGHlsl.h
@@ -29,8 +29,8 @@ public:
 	virtual FString GetRuntimeName() const override { return TEXT("NNERuntimeRDGHlsl"); };
 	virtual bool IsPlatformSupported(const ITargetPlatform* TargetPlatform) const override { return true; };
 
-	virtual bool CanCreateModelData(FString FileType, TConstArrayView<uint8> FileData) const override;
-	virtual TArray<uint8> CreateModelData(FString FileType, TConstArrayView<uint8> FileData) override;
+	virtual bool CanCreateModelData(FString FileType, TConstArrayView<uint8> FileData, FGuid FileId) const override;
+	virtual TArray<uint8> CreateModelData(FString FileType, TConstArrayView<uint8> FileData, FGuid FileId) override;
 
 	virtual bool CanCreateModelRDG(TObjectPtr<UNNEModelData> ModelData) const override;
 	virtual TUniquePtr<UE::NNE::IModelRDG> CreateModel(TObjectPtr<UNNEModelData> ModelData) override;
