Passing the target platform to model data creation to allow for platform specific compiled neural network models.

Problem: So far, all runtimes cooked models in the same maner for every supported target platform and thus no target information was required. Runtimes which compile models down to game code will need information about the target to set compiler flags or do other optimisations.

Solution: Pass the ITargetPlatform to CanCreateModelData and CreateModelData as additional argument.

#rnx

diff --git a/Engine/Plugins/Experimental/NNE/Source/NNE/Private/NNEModelData.cpp b/Engine/Plugins/Experimental/NNE/Source/NNE/Private/NNEModelData.cpp
--- a/Engine/Plugins/Experimental/NNE/Source/NNE/Private/NNEModelData.cpp
+++ b/Engine/Plugins/Experimental/NNE/Source/NNE/Private/NNEModelData.cpp
@@ -69,12 +69,12 @@ inline void PutIntoDDC(const FGuid& FileId, const FString& RuntimeName, FSharedB
 
 #endif
 
-inline TArray<uint8> CreateRuntimeDataBlob(const FString& RuntimeName, FString FileType, const TArray<uint8>& FileData, FGuid FileId)
+inline TArray<uint8> CreateRuntimeDataBlob(const FString& RuntimeName, FString FileType, const TArray<uint8>& FileData, FGuid FileId, const ITargetPlatform* TargetPlatform)
 {
 	TWeakInterfacePtr<INNERuntime> NNERuntime = UE::NNE::GetRuntime<INNERuntime>(RuntimeName);
 	if (NNERuntime.IsValid())
 	{
-		return NNERuntime->CreateModelData(FileType, FileData, FileId);
+		return NNERuntime->CreateModelData(FileType, FileData, FileId, TargetPlatform);
 	}
 	else
 	{
@@ -147,7 +147,7 @@ TConstArrayView<uint8> UNNEModelData::GetModelData(const FString& RuntimeName)
 #endif //WITH_EDITOR
 
 	// Try to create the model
-	TArray<uint8> CreatedData = CreateRuntimeDataBlob(RuntimeName, FileType, FileData, FileId);
+	TArray<uint8> CreatedData = CreateRuntimeDataBlob(RuntimeName, FileType, FileData, FileId, nullptr);
 	if (CreatedData.Num() < 1)
 	{
 		return {};
@@ -191,7 +191,7 @@ void UNNEModelData::Serialize(FArchive& Ar)
 
 		for (const FString& RuntimeName : CookedRuntimeNames)
 		{
-			TArray<uint8> CreatedData = CreateRuntimeDataBlob(RuntimeName, FileType, FileData, FileId);
+			TArray<uint8> CreatedData = CreateRuntimeDataBlob(RuntimeName, FileType, FileData, FileId, Ar.GetArchiveState().CookingTarget());
 			if (CreatedData.Num() > 0)
 			{
 				ModelData.Add(RuntimeName, CreatedData);
diff --git a/Engine/Plugins/Experimental/NNE/Source/NNE/Public/NNERuntime.h b/Engine/Plugins/Experimental/NNE/Source/NNE/Public/NNERuntime.h
--- a/Engine/Plugins/Experimental/NNE/Source/NNE/Public/NNERuntime.h
+++ b/Engine/Plugins/Experimental/NNE/Source/NNE/Public/NNERuntime.h
@@ -7,6 +7,8 @@
 
 #include "NNERuntime.generated.h"
 
+class ITargetPlatform;
+
 UINTERFACE()
 class NNE_API UNNERuntime : public UInterface
 {
@@ -49,9 +51,10 @@ public:
 	 * @param FileType The type of file inside FileData. Corresponds to the file extension (e.g. 'onnx').
 	 * @param FileData The raw binary file of a neural network model.
 	 * @param FileId The unique identifier representing FileData.
+	 * @param TargetPlatform The Interface identifying the target platform for which the data needs to be created. A null pointer indicates the currently compiled/running platform.
 	 * @return True if the runtime is able to create model data, false otherwise.
 	 */
-	virtual bool CanCreateModelData(FString FileType, TConstArrayView<uint8> FileData, FGuid FileId) const = 0;
+	virtual bool CanCreateModelData(FString FileType, TConstArrayView<uint8> FileData, FGuid FileId, const ITargetPlatform* TargetPlatform) const = 0;
 
 	/**
 	 * Create model data given some raw file data.
@@ -59,7 +62,8 @@ public:
 	 * @param FileType The type of file inside FileData. Corresponds to the file extension (e.g. 'onnx').
 	 * @param FileData The raw binary file of a neural network model.
 	 * @param FileId The unique identifier representing FileData.
+	 * @param TargetPlatform The Interface identifying the target platform for which the data needs to be created. A null pointer indicates the currently compiled/running platform.
 	 * @return Data representing the runtime specific representation of the model to be stored by UNNEModelData on success or an empty array otherwise.
 	 */
-	virtual TArray<uint8> CreateModelData(FString FileType, TConstArrayView<uint8> FileData, FGuid FileId) = 0;
+	virtual TArray<uint8> CreateModelData(FString FileType, TConstArrayView<uint8> FileData, FGuid FileId, const ITargetPlatform* TargetPlatform) = 0;
 };
\ No newline at end of file
diff --git a/Engine/Plugins/Experimental/NNERuntimeORTCpu/Source/NNERuntimeORTCpu/Private/NNERuntimeORTCpu.cpp b/Engine/Plugins/Experimental/NNERuntimeORTCpu/Source/NNERuntimeORTCpu/Private/NNERuntimeORTCpu.cpp
--- a/Engine/Plugins/Experimental/NNERuntimeORTCpu/Source/NNERuntimeORTCpu/Private/NNERuntimeORTCpu.cpp
+++ b/Engine/Plugins/Experimental/NNERuntimeORTCpu/Source/NNERuntimeORTCpu/Private/NNERuntimeORTCpu.cpp
@@ -16,14 +16,14 @@
 FGuid UNNERuntimeORTCpuImpl::GUID = FGuid((int32)'O', (int32)'C', (int32)'P', (int32)'U');
 int32 UNNERuntimeORTCpuImpl::Version = 0x00000001;
 
-bool UNNERuntimeORTCpuImpl::CanCreateModelData(FString FileType, TConstArrayView<uint8> FileData, FGuid FileId) const
+bool UNNERuntimeORTCpuImpl::CanCreateModelData(FString FileType, TConstArrayView<uint8> FileData, FGuid FileId, const ITargetPlatform* TargetPlatform) const
 {
 	return FileType.Compare("onnx", ESearchCase::IgnoreCase) == 0;
 }
 
-TArray<uint8> UNNERuntimeORTCpuImpl::CreateModelData(FString FileType, TConstArrayView<uint8> FileData, FGuid FileId)
+TArray<uint8> UNNERuntimeORTCpuImpl::CreateModelData(FString FileType, TConstArrayView<uint8> FileData, FGuid FileId, const ITargetPlatform* TargetPlatform)
 {
-	if (!CanCreateModelData(FileType, FileData, FileId))
+	if (!CanCreateModelData(FileType, FileData, FileId, TargetPlatform))
 	{
 		return {};
 	}
diff --git a/Engine/Plugins/Experimental/NNERuntimeORTCpu/Source/NNERuntimeORTCpu/Private/NNERuntimeORTCpu.h b/Engine/Plugins/Experimental/NNERuntimeORTCpu/Source/NNERuntimeORTCpu/Private/NNERuntimeORTCpu.h
--- a/Engine/Plugins/Experimental/NNERuntimeORTCpu/Source/NNERuntimeORTCpu/Private/NNERuntimeORTCpu.h
+++ b/Engine/Plugins/Experimental/NNERuntimeORTCpu/Source/NNERuntimeORTCpu/Private/NNERuntimeORTCpu.h
@@ -34,8 +34,8 @@ public:
 	virtual FString GetRuntimeName() const override { return TEXT("NNERuntimeORTCpu"); };
 	virtual bool IsPlatformSupported(const ITargetPlatform* TargetPlatform) const override { return true; };
 
-	virtual bool CanCreateModelData(FString FileType, TConstArrayView<uint8> FileData, FGuid FileId) const override;
-	virtual TArray<uint8> CreateModelData(FString FileType, TConstArrayView<uint8> FileData, FGuid FileId) override;
+	virtual bool CanCreateModelData(FString FileType, TConstArrayView<uint8> FileData, FGuid FileId, const ITargetPlatform* TargetPlatform) const override;
+	virtual TArray<uint8> CreateModelData(FString FileType, TConstArrayView<uint8> FileData, FGuid FileId, const ITargetPlatform* TargetPlatform) override;
 
 	virtual bool CanCreateModelCPU(TObjectPtr<UNNEModelData> ModelData) const override;
 	virtual TUniquePtr<UE::NNE::IModelCPU> CreateModel(TObjectPtr<UNNEModelData> ModelData) override;
diff --git a/Engine/Plugins/Experimental/NNERuntimeORTGpu/Source/NNERuntimeORT/Private/NNERuntimeORT.cpp b/Engine/Plugins/Experimental/NNERuntimeORTGpu/Source/NNERuntimeORT/Private/NNERuntimeORT.cpp
--- a/Engine/Plugins/Experimental/NNERuntimeORTGpu/Source/NNERuntimeORT/Private/NNERuntimeORT.cpp
+++ b/Engine/Plugins/Experimental/NNERuntimeORTGpu/Source/NNERuntimeORT/Private/NNERuntimeORT.cpp
@@ -18,14 +18,14 @@
 FGuid UNNERuntimeORTGpuImpl::GUID = FGuid((int32)'O', (int32)'G', (int32)'P', (int32)'U');
 int32 UNNERuntimeORTGpuImpl::Version = 0x00000001;
 
-bool UNNERuntimeORTGpuImpl::CanCreateModelData(FString FileType, TConstArrayView<uint8> FileData, FGuid FileId) const
+bool UNNERuntimeORTGpuImpl::CanCreateModelData(FString FileType, TConstArrayView<uint8> FileData, FGuid FileId, const ITargetPlatform* TargetPlatform) const
 {
 	return FileType.Compare("onnx", ESearchCase::IgnoreCase) == 0;
 }
 
-TArray<uint8> UNNERuntimeORTGpuImpl::CreateModelData(FString FileType, TConstArrayView<uint8> FileData, FGuid FileId)
+TArray<uint8> UNNERuntimeORTGpuImpl::CreateModelData(FString FileType, TConstArrayView<uint8> FileData, FGuid FileId, const ITargetPlatform* TargetPlatform)
 {
-	if (!CanCreateModelData(FileType, FileData, FileId))
+	if (!CanCreateModelData(FileType, FileData, FileId, TargetPlatform))
 	{
 		return {};
 	}
diff --git a/Engine/Plugins/Experimental/NNERuntimeORTGpu/Source/NNERuntimeORT/Private/NNERuntimeORT.h b/Engine/Plugins/Experimental/NNERuntimeORTGpu/Source/NNERuntimeORT/Private/NNERuntimeORT.h
--- a/Engine/Plugins/Experimental/NNERuntimeORTGpu/Source/NNERuntimeORT/Private/NNERuntimeORT.h
+++ b/Engine/Plugins/Experimental/NNERuntimeORTGpu/Source/NNERuntimeORT/Private/NNERuntimeORT.h
@@ -46,8 +46,8 @@ public:
 	virtual FString GetRuntimeName() const override;
 	virtual bool IsPlatformSupported(const ITargetPlatform* TargetPlatform) const override { return true; };
 
-	virtual bool CanCreateModelData(FString FileType, TConstArrayView<uint8> FileData, FGuid FileId) const override;
-	virtual TArray<uint8> CreateModelData(FString FileType, TConstArrayView<uint8> FileData, FGuid FileId) override;
+	virtual bool CanCreateModelData(FString FileType, TConstArrayView<uint8> FileData, FGuid FileId, const ITargetPlatform* TargetPlatform) const override;
+	virtual TArray<uint8> CreateModelData(FString FileType, TConstArrayView<uint8> FileData, FGuid FileId, const ITargetPlatform* TargetPlatform) override;
 
 	virtual bool CanCreateModelGPU(TObjectPtr<UNNEModelData> ModelData) const override;
 	virtual TUniquePtr<UE::NNE::IModelGPU> CreateModel(TObjectPtr<UNNEModelData> ModelData) override;
diff --git a/Engine/Plugins/Experimental/NNERuntimeRDG/Source/NNERuntimeRDG/Private/NNERuntimeRDGDml.cpp b/Engine/Plugins/Experimental/NNERuntimeRDG/Source/NNERuntimeRDG/Private/NNERuntimeRDGDml.cpp
--- a/Engine/Plugins/Experimental/NNERuntimeRDG/Source/NNERuntimeRDG/Private/NNERuntimeRDGDml.cpp
+++ b/Engine/Plugins/Experimental/NNERuntimeRDG/Source/NNERuntimeRDG/Private/NNERuntimeRDGDml.cpp
@@ -96,14 +96,14 @@ bool FRuntimeDmlStartup()
 
 using namespace UE::NNERuntimeRDG::Private::Dml;
 
-bool UNNERuntimeRDGDmlImpl::CanCreateModelData(FString FileType, TConstArrayView<uint8> FileData, FGuid FileId) const
+bool UNNERuntimeRDGDmlImpl::CanCreateModelData(FString FileType, TConstArrayView<uint8> FileData, FGuid FileId, const ITargetPlatform* TargetPlatform) const
 {
 	return FileType.Compare("onnx", ESearchCase::IgnoreCase) == 0;
 }
 
-TArray<uint8> UNNERuntimeRDGDmlImpl::CreateModelData(FString FileType, TConstArrayView<uint8> FileData, FGuid FileId)
+TArray<uint8> UNNERuntimeRDGDmlImpl::CreateModelData(FString FileType, TConstArrayView<uint8> FileData, FGuid FileId, const ITargetPlatform* TargetPlatform)
 {
-	if (!CanCreateModelData(FileType, FileData, FileId))
+	if (!CanCreateModelData(FileType, FileData, FileId, TargetPlatform))
 	{
 		return {};
 	}
diff --git a/Engine/Plugins/Experimental/NNERuntimeRDG/Source/NNERuntimeRDG/Private/NNERuntimeRDGDml.h b/Engine/Plugins/Experimental/NNERuntimeRDG/Source/NNERuntimeRDG/Private/NNERuntimeRDGDml.h
--- a/Engine/Plugins/Experimental/NNERuntimeRDG/Source/NNERuntimeRDG/Private/NNERuntimeRDGDml.h
+++ b/Engine/Plugins/Experimental/NNERuntimeRDG/Source/NNERuntimeRDG/Private/NNERuntimeRDGDml.h
@@ -34,8 +34,8 @@ public:
 	virtual FString GetRuntimeName() const override { return TEXT("NNERuntimeRDGDml"); };
 	virtual bool IsPlatformSupported(const ITargetPlatform* TargetPlatform) const override { return true; };
 
-	virtual bool CanCreateModelData(FString FileType, TConstArrayView<uint8> FileData, FGuid FileId) const override;
-	virtual TArray<uint8> CreateModelData(FString FileType, TConstArrayView<uint8> FileData, FGuid FileId) override;
+	virtual bool CanCreateModelData(FString FileType, TConstArrayView<uint8> FileData, FGuid FileId, const ITargetPlatform* TargetPlatform) const override;
+	virtual TArray<uint8> CreateModelData(FString FileType, TConstArrayView<uint8> FileData, FGuid FileId, const ITargetPlatform* TargetPlatform) override;
 
 	virtual bool CanCreateModelRDG(TObjectPtr<UNNEModelData> ModelData) const override;
 	virtual TUniquePtr<UE::NNE::IModelRDG> CreateModel(TObjectPtr<UNNEModelData> ModelData) override;
diff --git a/Engine/Plugins/Experimental/NNERuntimeRDG/Source/NNERuntimeRDG/Private/NNERuntimeRDGHlsl.cpp b/Engine/Plugins/Experimental/NNERuntimeRDG/Source/NNERuntimeRDG/Private/NNERuntimeRDGHlsl.cpp
--- a/Engine/Plugins/Experimental/NNERuntimeRDG/Source/NNERuntimeRDG/Private/NNERuntimeRDGHlsl.cpp
+++ b/Engine/Plugins/Experimental/NNERuntimeRDG/Source/NNERuntimeRDG/Private/NNERuntimeRDGHlsl.cpp
@@ -73,7 +73,7 @@ bool UNNERuntimeRDGHlslImpl::Init()
 	return true;
 }
 
-bool UNNERuntimeRDGHlslImpl::CanCreateModelData(FString FileType, TConstArrayView<uint8> FileData, FGuid FileId) const
+bool UNNERuntimeRDGHlslImpl::CanCreateModelData(FString FileType, TConstArrayView<uint8> FileData, FGuid FileId, const ITargetPlatform* TargetPlatform) const
 {
 	return FileType.Compare("onnx", ESearchCase::IgnoreCase) == 0;
 }
@@ -93,9 +93,9 @@ bool UNNERuntimeRDGHlslImpl::CanCreateModelRDG(TObjectPtr<UNNEModelData> ModelDa
 	return bResult;
 };
 
-TArray<uint8> UNNERuntimeRDGHlslImpl::CreateModelData(FString FileType, TConstArrayView<uint8> FileData, FGuid FileId)
+TArray<uint8> UNNERuntimeRDGHlslImpl::CreateModelData(FString FileType, TConstArrayView<uint8> FileData, FGuid FileId, const ITargetPlatform* TargetPlatform)
 {
-	if (!CanCreateModelData(FileType, FileData, FileId))
+	if (!CanCreateModelData(FileType, FileData, FileId, TargetPlatform))
 	{
 		return {};
 	}
diff --git a/Engine/Plugins/Experimental/NNERuntimeRDG/Source/NNERuntimeRDG/Private/NNERuntimeRDGHlsl.h b/Engine/Plugins/Experimental/NNERuntimeRDG/Source/NNERuntimeRDG/Private/NNERuntimeRDGHlsl.h
--- a/Engine/Plugins/Experimental/NNERuntimeRDG/Source/NNERuntimeRDG/Private/NNERuntimeRDGHlsl.h
+++ b/Engine/Plugins/Experimental/NNERuntimeRDG/Source/NNERuntimeRDG/Private/NNERuntimeRDGHlsl.h
@@ -29,8 +29,8 @@ public:
 	virtual FString GetRuntimeName() const override { return TEXT("NNERuntimeRDGHlsl"); };
 	virtual bool IsPlatformSupported(const ITargetPlatform* TargetPlatform) const override { return true; };
 
-	virtual bool CanCreateModelData(FString FileType, TConstArrayView<uint8> FileData, FGuid FileId) const override;
-	virtual TArray<uint8> CreateModelData(FString FileType, TConstArrayView<uint8> FileData, FGuid FileId) override;
+	virtual bool CanCreateModelData(FString FileType, TConstArrayView<uint8> FileData, FGuid FileId, const ITargetPlatform* TargetPlatform) const override;
+	virtual TArray<uint8> CreateModelData(FString FileType, TConstArrayView<uint8> FileData, FGuid FileId, const ITargetPlatform* TargetPlatform) override;
 
 	virtual bool CanCreateModelRDG(TObjectPtr<UNNEModelData> ModelData) const override;
 	virtual TUniquePtr<UE::NNE::IModelRDG> CreateModel(TObjectPtr<UNNEModelData> ModelData) override;
