NNE: Fix NNE DDC cache to avoid version conflicts in the MHA plugin.

Problem:
The Meta Human Animator team noticed a problem when loading NNE model data assets from the DDC cache due to different versions being stored under the same hash.

Solution:
This CL revisits the caching strategy of NNE model data making them uniquely identifiable by runtime+asset and letting the runtime includes the data version and, if needed, the target platform. For this the NNE API is extended, the change is backward compatible.

Test:
NNE smoketest + Bumped runtime model data version and verify the correct request are processed without collision and the model data is recreated on absence in the DDC.

Based on the shelve 26363260 from nico.ranieri

#rnx

diff --git a/Engine/Plugins/Experimental/NNE/Source/NNE/Private/NNEModelData.cpp b/Engine/Plugins/Experimental/NNE/Source/NNE/Private/NNEModelData.cpp
--- a/Engine/Plugins/Experimental/NNE/Source/NNE/Private/NNEModelData.cpp
+++ b/Engine/Plugins/Experimental/NNE/Source/NNE/Private/NNEModelData.cpp
@@ -35,17 +35,25 @@ FCustomVersionRegistration NNEModelDataVersion(UNNEModelData::GUID, LatestVersio
 
 #if WITH_EDITOR
 
-inline UE::DerivedData::FCacheKey CreateCacheKey(const FGuid& FileId, const FString& RuntimeName)
+inline FString GetDDCRequestId(const FString& FileId, const FString& RuntimeName, const FString& ModelDataIdentifier)
 {
-	FString GuidString = FileId.ToString(EGuidFormats::Digits);
-	return { UE::DerivedData::FCacheBucket(FWideStringView(*GuidString)), FIoHash::HashBuffer(MakeMemoryView(FTCHARToUTF8(RuntimeName))) };
+	//RuntimeName and FileId are embedded to the id to ensure no potential collision between the runtime/assets
+	return RuntimeName + "-" + FileId + "-" + ModelDataIdentifier;
 }
 
-inline FSharedBuffer GetFromDDC(const FGuid& FileId, const FString& RuntimeName)
+inline UE::DerivedData::FCacheKey CreateCacheKey(const FString& FileId, const FString& RequestId)
 {
+	return { UE::DerivedData::FCacheBucket(FWideStringView(*FileId)), FIoHash::HashBuffer(MakeMemoryView(FTCHARToUTF8(RequestId))) };
+}
+
+inline FSharedBuffer GetFromDDC(const FGuid& FileId, const FString& RuntimeName, const FString& ModelDataIdentifier)
+{
+	FString FileIdStr = FileId.ToString(EGuidFormats::Digits);
+	FString RequestId = GetDDCRequestId(FileIdStr, RuntimeName, ModelDataIdentifier);
+
 	UE::DerivedData::FCacheGetValueRequest GetRequest;
-	GetRequest.Name = FString("Get-") + RuntimeName + FString("-") + FileId.ToString(EGuidFormats::Digits);
-	GetRequest.Key = CreateCacheKey(FileId,  RuntimeName);
+	GetRequest.Name = FString("Get-") + RequestId;
+	GetRequest.Key = CreateCacheKey(FileIdStr, RequestId);
 	FSharedBuffer RawDerivedData;
 	UE::DerivedData::FRequestOwner BlockingGetOwner(UE::DerivedData::EPriority::Blocking);
 	UE::DerivedData::GetCache().GetValue({ GetRequest }, BlockingGetOwner, [&RawDerivedData](UE::DerivedData::FCacheGetValueResponse&& Response)
@@ -56,11 +64,14 @@ inline FSharedBuffer GetFromDDC(const FGuid& FileId, const FString& RuntimeName)
 	return RawDerivedData;
 }
 
-inline void PutIntoDDC(const FGuid& FileId, const FString& RuntimeName, FSharedBuffer& Data)
+inline void PutIntoDDC(const FGuid& FileId, const FString& RuntimeName, const FString& ModelDataIdentifier, FSharedBuffer& Data)
 {
+	FString FileIdStr = FileId.ToString(EGuidFormats::Digits);
+	FString RequestId = GetDDCRequestId(FileIdStr, RuntimeName, ModelDataIdentifier);
+
 	UE::DerivedData::FCachePutValueRequest PutRequest;
-	PutRequest.Name = FString("Put-") + RuntimeName + FString("-") + FileId.ToString(EGuidFormats::Digits);
-	PutRequest.Key = CreateCacheKey(FileId, RuntimeName);
+	PutRequest.Name = FString("Put-") + RequestId;
+	PutRequest.Key = CreateCacheKey(FileIdStr, RequestId);
 	PutRequest.Value = UE::DerivedData::FValue::Compress(Data);
 	UE::DerivedData::FRequestOwner BlockingPutOwner(UE::DerivedData::EPriority::Blocking);
 	UE::DerivedData::GetCache().PutValue({ PutRequest }, BlockingPutOwner);
@@ -135,8 +146,22 @@ TConstArrayView<uint8> UNNEModelData::GetModelData(const FString& RuntimeName)
 	}
 	
 #if WITH_EDITOR
-	// Check if we have a remote cache hit
-	FSharedBuffer RemoteData = GetFromDDC(FileId, RuntimeName);
+	TWeakInterfacePtr<INNERuntime> NNERuntime = UE::NNE::GetRuntime<INNERuntime>(RuntimeName);
+	if (!NNERuntime.IsValid())
+	{
+		UE_LOG(LogNNE, Error, TEXT("UNNEModelData: Runtime '%s' is among the target runtimes but instance is invalid."), *RuntimeName);
+		return {};
+	}
+
+	FString ModelDataIdentifier = NNERuntime->GetModelDataIdentifier(FileType, FileData, FileId, nullptr);
+	if (ModelDataIdentifier.Len() == 0)
+	{
+		UE_LOG(LogNNE, Error, TEXT("UNNEModelData: Runtime '%s' returned an empty string as a ModelDataIdentifier. GetModelDataIdentifier should always return a valid identifier."), *RuntimeName);
+		return {};
+	}
+	
+	// Check if we have a DDC cache hit
+	FSharedBuffer RemoteData = GetFromDDC(FileId, RuntimeName, ModelDataIdentifier);
 	if (RemoteData.GetSize() > 0)
 	{
 		ModelData.Add(RuntimeName, TArray<uint8>((uint8*)RemoteData.GetData(), RemoteData.GetSize()));
@@ -159,7 +184,7 @@ TConstArrayView<uint8> UNNEModelData::GetModelData(const FString& RuntimeName)
 #if WITH_EDITOR
 	// And put it into DDC
 	FSharedBuffer SharedBuffer = MakeSharedBufferFromArray(MoveTemp(CreatedData));
-	PutIntoDDC(FileId, RuntimeName, SharedBuffer);
+	PutIntoDDC(FileId, RuntimeName, ModelDataIdentifier, SharedBuffer);
 #endif //WITH_EDITOR
 	
 	TArray<uint8>* CachedCreatedData = ModelData.Find(RuntimeName);
@@ -196,8 +221,24 @@ void UNNEModelData::Serialize(FArchive& Ar)
 			{
 				ModelData.Add(RuntimeName, CreatedData);
 #if WITH_EDITOR
-				FSharedBuffer SharedBuffer = MakeSharedBufferFromArray(MoveTemp(CreatedData));
-				PutIntoDDC(FileId, RuntimeName, SharedBuffer);
+				TWeakInterfacePtr<INNERuntime> NNERuntime = UE::NNE::GetRuntime<INNERuntime>(RuntimeName);
+				if (NNERuntime.IsValid())
+				{
+					FString ModelDataIdentifier = NNERuntime->GetModelDataIdentifier(FileType, FileData, FileId, Ar.GetArchiveState().CookingTarget());
+					if (ModelDataIdentifier.Len() > 0)
+					{
+						FSharedBuffer SharedBuffer = MakeSharedBufferFromArray(MoveTemp(CreatedData));
+						PutIntoDDC(FileId, RuntimeName, ModelDataIdentifier, SharedBuffer);
+					}
+					else
+					{
+						UE_LOG(LogNNE, Warning, TEXT("UNNEModelData: Runtime '%s' returned an empty string as a ModelDataIdentifier while cooking. GetModelDataIdentifier should always return a valid identifier."), *RuntimeName);
+					}
+				}
+				else
+				{
+					UE_LOG(LogNNE, Warning, TEXT("UNNEModelData: Runtime '%s' is among the cooked runtimes but instance is invalid."), *RuntimeName);
+				}
 #endif //WITH_EDITOR
 			}
 		}
diff --git a/Engine/Plugins/Experimental/NNE/Source/NNE/Public/NNERuntime.h b/Engine/Plugins/Experimental/NNE/Source/NNE/Public/NNERuntime.h
--- a/Engine/Plugins/Experimental/NNE/Source/NNE/Public/NNERuntime.h
+++ b/Engine/Plugins/Experimental/NNE/Source/NNE/Public/NNERuntime.h
@@ -66,4 +66,15 @@ public:
 	 * @return Data representing the runtime specific representation of the model to be stored by UNNEModelData on success or an empty array otherwise.
 	 */
 	virtual TArray<uint8> CreateModelData(FString FileType, TConstArrayView<uint8> FileData, FGuid FileId, const ITargetPlatform* TargetPlatform) = 0;
+
+	/**
+	 * Get an id uniquely identifying the model data. This is used by UNNEModelData to create a FCacheKey to get and add the model data from and to a DDC.
+	 *
+	 * @param FileType The type of file inside FileData. Corresponds to the file extension (e.g. 'onnx').
+	 * @param FileData The raw binary file of a neural network model.
+	 * @param FileId The unique identifier representing FileData.
+	 * @param TargetPlatform The Interface identifying the target platform for which the data is produced or by which it will be consumed. A null pointer indicates the currently compiled/running platform.
+	 * @return A string that identifies the model data uniquely, e.g. by containing the FileId, model data version and target platform display name.
+	 */
+	virtual FString GetModelDataIdentifier(FString FileType, TConstArrayView<uint8> FileData, FGuid FileId, const ITargetPlatform* TargetPlatform) = 0;
 };
\ No newline at end of file
diff --git a/Engine/Plugins/Experimental/NNERuntimeORTCpu/Source/NNERuntimeORTCpu/Private/NNERuntimeORTCpu.cpp b/Engine/Plugins/Experimental/NNERuntimeORTCpu/Source/NNERuntimeORTCpu/Private/NNERuntimeORTCpu.cpp
--- a/Engine/Plugins/Experimental/NNERuntimeORTCpu/Source/NNERuntimeORTCpu/Private/NNERuntimeORTCpu.cpp
+++ b/Engine/Plugins/Experimental/NNERuntimeORTCpu/Source/NNERuntimeORTCpu/Private/NNERuntimeORTCpu.cpp
@@ -50,6 +50,11 @@ TArray<uint8> UNNERuntimeORTCpuImpl::CreateModelData(FString FileType, TConstArr
 	return Result;
 }
 
+FString UNNERuntimeORTCpuImpl::GetModelDataIdentifier(FString FileType, TConstArrayView<uint8> FileData, FGuid FileId, const ITargetPlatform* TargetPlatform)
+{
+	return FileId.ToString(EGuidFormats::Digits) + "-" + UNNERuntimeORTCpuImpl::GUID.ToString(EGuidFormats::Digits) + "-" + FString::FromInt(UNNERuntimeORTCpuImpl::Version);
+}
+
 bool UNNERuntimeORTCpuImpl::CanCreateModelCPU(TObjectPtr<UNNEModelData> ModelData) const
 {
 	check(ModelData != nullptr);
diff --git a/Engine/Plugins/Experimental/NNERuntimeORTCpu/Source/NNERuntimeORTCpu/Private/NNERuntimeORTCpu.h b/Engine/Plugins/Experimental/NNERuntimeORTCpu/Source/NNERuntimeORTCpu/Private/NNERuntimeORTCpu.h
--- a/Engine/Plugins/Experimental/NNERuntimeORTCpu/Source/NNERuntimeORTCpu/Private/NNERuntimeORTCpu.h
+++ b/Engine/Plugins/Experimental/NNERuntimeORTCpu/Source/NNERuntimeORTCpu/Private/NNERuntimeORTCpu.h
@@ -36,6 +36,7 @@ public:
 
 	virtual bool CanCreateModelData(FString FileType, TConstArrayView<uint8> FileData, FGuid FileId, const ITargetPlatform* TargetPlatform) const override;
 	virtual TArray<uint8> CreateModelData(FString FileType, TConstArrayView<uint8> FileData, FGuid FileId, const ITargetPlatform* TargetPlatform) override;
+	virtual FString GetModelDataIdentifier(FString FileType, TConstArrayView<uint8> FileData, FGuid FileId, const ITargetPlatform* TargetPlatform) override;
 
 	virtual bool CanCreateModelCPU(TObjectPtr<UNNEModelData> ModelData) const override;
 	virtual TUniquePtr<UE::NNE::IModelCPU> CreateModel(TObjectPtr<UNNEModelData> ModelData) override;
diff --git a/Engine/Plugins/Experimental/NNERuntimeORTGpu/Source/NNERuntimeORT/Private/NNERuntimeORT.cpp b/Engine/Plugins/Experimental/NNERuntimeORTGpu/Source/NNERuntimeORT/Private/NNERuntimeORT.cpp
--- a/Engine/Plugins/Experimental/NNERuntimeORTGpu/Source/NNERuntimeORT/Private/NNERuntimeORT.cpp
+++ b/Engine/Plugins/Experimental/NNERuntimeORTGpu/Source/NNERuntimeORT/Private/NNERuntimeORT.cpp
@@ -52,6 +52,11 @@ TArray<uint8> UNNERuntimeORTGpuImpl::CreateModelData(FString FileType, TConstArr
 	return Result;
 }
 
+FString UNNERuntimeORTGpuImpl::GetModelDataIdentifier(FString FileType, TConstArrayView<uint8> FileData, FGuid FileId, const ITargetPlatform* TargetPlatform)
+{
+	return FileId.ToString(EGuidFormats::Digits) + "-" + UNNERuntimeORTGpuImpl::GUID.ToString(EGuidFormats::Digits) + "-" + FString::FromInt(UNNERuntimeORTGpuImpl::Version);
+}
+
 void UNNERuntimeORTGpuImpl::Init(ENNERuntimeORTGpuProvider InProvider)
 {
 	check(!ORTEnvironment.IsValid());
diff --git a/Engine/Plugins/Experimental/NNERuntimeORTGpu/Source/NNERuntimeORT/Private/NNERuntimeORT.h b/Engine/Plugins/Experimental/NNERuntimeORTGpu/Source/NNERuntimeORT/Private/NNERuntimeORT.h
--- a/Engine/Plugins/Experimental/NNERuntimeORTGpu/Source/NNERuntimeORT/Private/NNERuntimeORT.h
+++ b/Engine/Plugins/Experimental/NNERuntimeORTGpu/Source/NNERuntimeORT/Private/NNERuntimeORT.h
@@ -48,6 +48,7 @@ public:
 
 	virtual bool CanCreateModelData(FString FileType, TConstArrayView<uint8> FileData, FGuid FileId, const ITargetPlatform* TargetPlatform) const override;
 	virtual TArray<uint8> CreateModelData(FString FileType, TConstArrayView<uint8> FileData, FGuid FileId, const ITargetPlatform* TargetPlatform) override;
+	virtual FString GetModelDataIdentifier(FString FileType, TConstArrayView<uint8> FileData, FGuid FileId, const ITargetPlatform* TargetPlatform) override;
 
 	virtual bool CanCreateModelGPU(TObjectPtr<UNNEModelData> ModelData) const override;
 	virtual TUniquePtr<UE::NNE::IModelGPU> CreateModel(TObjectPtr<UNNEModelData> ModelData) override;
diff --git a/Engine/Plugins/Experimental/NNERuntimeRDG/Source/NNERuntimeRDG/Private/NNERuntimeRDGDml.cpp b/Engine/Plugins/Experimental/NNERuntimeRDG/Source/NNERuntimeRDG/Private/NNERuntimeRDGDml.cpp
--- a/Engine/Plugins/Experimental/NNERuntimeRDG/Source/NNERuntimeRDG/Private/NNERuntimeRDGDml.cpp
+++ b/Engine/Plugins/Experimental/NNERuntimeRDG/Source/NNERuntimeRDG/Private/NNERuntimeRDGDml.cpp
@@ -135,6 +135,11 @@ TArray<uint8> UNNERuntimeRDGDmlImpl::CreateModelData(FString FileType, TConstArr
 	return Result;
 };
 
+FString UNNERuntimeRDGDmlImpl::GetModelDataIdentifier(FString FileType, TConstArrayView<uint8> FileData, FGuid FileId, const ITargetPlatform* TargetPlatform)
+{
+	return FileId.ToString(EGuidFormats::Digits) + "-" + FModelInfo::Get()->GetGuid().ToString(EGuidFormats::Digits) + "-" + FString::FromInt(FModelInfo::Get()->GetVersion());
+}
+
 bool UNNERuntimeRDGDmlImpl::CanCreateModelRDG(TObjectPtr<UNNEModelData> ModelData) const
 {
 #ifdef NNE_USE_DIRECTML
diff --git a/Engine/Plugins/Experimental/NNERuntimeRDG/Source/NNERuntimeRDG/Private/NNERuntimeRDGDml.h b/Engine/Plugins/Experimental/NNERuntimeRDG/Source/NNERuntimeRDG/Private/NNERuntimeRDGDml.h
--- a/Engine/Plugins/Experimental/NNERuntimeRDG/Source/NNERuntimeRDG/Private/NNERuntimeRDGDml.h
+++ b/Engine/Plugins/Experimental/NNERuntimeRDG/Source/NNERuntimeRDG/Private/NNERuntimeRDGDml.h
@@ -36,6 +36,7 @@ public:
 
 	virtual bool CanCreateModelData(FString FileType, TConstArrayView<uint8> FileData, FGuid FileId, const ITargetPlatform* TargetPlatform) const override;
 	virtual TArray<uint8> CreateModelData(FString FileType, TConstArrayView<uint8> FileData, FGuid FileId, const ITargetPlatform* TargetPlatform) override;
+	virtual FString GetModelDataIdentifier(FString FileType, TConstArrayView<uint8> FileData, FGuid FileId, const ITargetPlatform* TargetPlatform) override;
 
 	virtual bool CanCreateModelRDG(TObjectPtr<UNNEModelData> ModelData) const override;
 	virtual TUniquePtr<UE::NNE::IModelRDG> CreateModel(TObjectPtr<UNNEModelData> ModelData) override;
diff --git a/Engine/Plugins/Experimental/NNERuntimeRDG/Source/NNERuntimeRDG/Private/NNERuntimeRDGHlsl.cpp b/Engine/Plugins/Experimental/NNERuntimeRDG/Source/NNERuntimeRDG/Private/NNERuntimeRDGHlsl.cpp
--- a/Engine/Plugins/Experimental/NNERuntimeRDG/Source/NNERuntimeRDG/Private/NNERuntimeRDGHlsl.cpp
+++ b/Engine/Plugins/Experimental/NNERuntimeRDG/Source/NNERuntimeRDG/Private/NNERuntimeRDGHlsl.cpp
@@ -122,6 +122,11 @@ TArray<uint8> UNNERuntimeRDGHlslImpl::CreateModelData(FString FileType, TConstAr
 	return Result;
 };
 
+FString UNNERuntimeRDGHlslImpl::GetModelDataIdentifier(FString FileType, TConstArrayView<uint8> FileData, FGuid FileId, const ITargetPlatform* TargetPlatform)
+{
+	return FileId.ToString(EGuidFormats::Digits) + "-" + UNNERuntimeRDGHlslImpl::GUID.ToString(EGuidFormats::Digits) + "-" + FString::FromInt(UNNERuntimeRDGHlslImpl::Version);
+}
+
 TUniquePtr<UE::NNE::IModelRDG> UNNERuntimeRDGHlslImpl::CreateModel(TObjectPtr<UNNEModelData> ModelData)
 {
 	if (!CanCreateModelRDG(ModelData))
diff --git a/Engine/Plugins/Experimental/NNERuntimeRDG/Source/NNERuntimeRDG/Private/NNERuntimeRDGHlsl.h b/Engine/Plugins/Experimental/NNERuntimeRDG/Source/NNERuntimeRDG/Private/NNERuntimeRDGHlsl.h
--- a/Engine/Plugins/Experimental/NNERuntimeRDG/Source/NNERuntimeRDG/Private/NNERuntimeRDGHlsl.h
+++ b/Engine/Plugins/Experimental/NNERuntimeRDG/Source/NNERuntimeRDG/Private/NNERuntimeRDGHlsl.h
@@ -31,6 +31,7 @@ public:
 
 	virtual bool CanCreateModelData(FString FileType, TConstArrayView<uint8> FileData, FGuid FileId, const ITargetPlatform* TargetPlatform) const override;
 	virtual TArray<uint8> CreateModelData(FString FileType, TConstArrayView<uint8> FileData, FGuid FileId, const ITargetPlatform* TargetPlatform) override;
+	virtual FString GetModelDataIdentifier(FString FileType, TConstArrayView<uint8> FileData, FGuid FileId, const ITargetPlatform* TargetPlatform) override;
 
 	virtual bool CanCreateModelRDG(TObjectPtr<UNNEModelData> ModelData) const override;
 	virtual TUniquePtr<UE::NNE::IModelRDG> CreateModel(TObjectPtr<UNNEModelData> ModelData) override;
